[GeometryMask] Moved plugin

#virtualized

diff --git a/Engine/Plugins/Experimental/GeometryMask/GeometryMask.uplugin b/Engine/Plugins/Experimental/GeometryMask/GeometryMask.uplugin
new file mode 100644
--- /dev/null
+++ b/Engine/Plugins/Experimental/GeometryMask/GeometryMask.uplugin
@@ -0,0 +1,28 @@
+﻿{
+	"FileVersion": 3,
+	"Version": 1,
+	"VersionName": "1.0",
+	"FriendlyName": "GeometryMask",
+	"Description": "",
+	"Category": "Other",
+	"CreatedBy": "",
+	"CreatedByURL": "",
+	"DocsURL": "",
+	"MarketplaceURL": "",
+	"CanContainContent": true,
+	"IsBetaVersion": false,
+	"IsExperimentalVersion": false,
+	"Installed": false,
+	"Modules": [
+		{
+			"Name": "GeometryMask",
+			"Type": "Runtime",
+			"LoadingPhase": "PostConfigInit"
+		},
+		{
+			"Name": "GeometryMaskEditor",
+			"Type": "Editor",
+			"LoadingPhase": "Default"
+		}
+	]
+}
diff --git a/Engine/Plugins/Experimental/GeometryMask/Shaders/Private/GeometryMaskBlurCS.usf b/Engine/Plugins/Experimental/GeometryMask/Shaders/Private/GeometryMaskBlurCS.usf
new file mode 100644
--- /dev/null
+++ b/Engine/Plugins/Experimental/GeometryMask/Shaders/Private/GeometryMaskBlurCS.usf
@@ -0,0 +1,91 @@
+// Copyright Epic Games, Inc. All Rights Reserved.
+
+#include "/Engine/Private/Common.ush"
+#include "/Engine/Public/Platform.ush"
+
+#define TILE_SIZE 8
+#define MAX_SAMPLES 127
+#define HALF_MAX_SAMPLES MAX_SAMPLES / 2
+
+DECLARE_SCALAR_ARRAY(int, SampleCount, 4);
+float4 WeightsAndOffsets[4 * HALF_MAX_SAMPLES];
+float4 BufferSizeAndDirection[4];
+
+uint2 InputDimensions;
+Texture2D InputTexture;
+SamplerState InputSampler;
+RWTexture2D<float4> OutputTexture;
+
+static const float4 ChannelMasks[4] = { float4(1,0,0,1), float4(0,1,0,1), float4(0,0,1,1), float4(0,0,0,1) };
+
+float4 GetSample(int ChannelIdx, float Weight, float Offset, float2 UV)
+{
+	const float2 MinUV = float2(0, 0);
+	const float2 MaxUV = float2(1.0, 1.0);
+	const float2 Direction = BufferSizeAndDirection[ChannelIdx].zw;
+	const float2 BufferSize = BufferSizeAndDirection[ChannelIdx].xy;
+	UV *= BufferSize;
+	
+	const float2 UVOffset = float2(Offset * BufferSize.x * Direction.x, Offset * BufferSize.y * Direction.y);
+	float4 ChannelMask = ChannelMasks[ChannelIdx];
+
+	return (InputTexture.SampleLevel(InputSampler, clamp(UV + UVOffset, MinUV, MaxUV), 0) * Weight
+		 + InputTexture.SampleLevel(InputSampler, clamp(UV - UVOffset, MinUV, MaxUV), 0) * Weight);
+}
+
+// Loosely based on Engine/Shaders/Private/SlatePostProcessPixelShader.usf
+[numthreads(TILE_SIZE, TILE_SIZE, 1)]
+void MainCS(
+	uint3 DispatchThreadId : SV_DispatchThreadID,
+	uint3 GroupThreadId : SV_GroupThreadID,
+	uint3 GroupId : SV_GroupID)
+{
+	uint2 PixelCoord = DispatchThreadId.xy;
+
+	float X = (float)PixelCoord.x + 0.5;
+	float Y = (float)PixelCoord.y + 0.5;
+	float2 UV = float2(X, Y);
+
+	const float2 MinUV = float2(0, 0);
+	const float2 MaxUV = float2(1.0, 1.0);
+
+	float4 Color = InputTexture.SampleLevel(InputSampler, clamp(UV, MinUV, MaxUV), 0);
+
+    // Perform on each of R,G,B,A
+    [unroll]
+    for (int ChannelIdx = 0; ChannelIdx < 4; ++ChannelIdx)
+    {
+    	int WeightsAndOffsetsIdx = ChannelIdx * HALF_MAX_SAMPLES;
+    	
+    	// First offset is in zw
+	    {
+	    	const float Weight = WeightsAndOffsets[WeightsAndOffsetsIdx].z;
+	    	const float Offset = WeightsAndOffsets[WeightsAndOffsetsIdx].w;
+
+	    	Color += GetSample(ChannelIdx, Weight, Offset, UV);
+	    }
+
+    	int ChannelSampleCount = GET_SCALAR_ARRAY_ELEMENT(SampleCount, ChannelIdx);
+    	for (int SampleIdx = 2; SampleIdx < ChannelSampleCount; SampleIdx += 2)
+    	{
+    		int ParamIdx = SampleIdx / 2;
+    		{
+    			const float Weight = WeightsAndOffsets[WeightsAndOffsetsIdx + ParamIdx].x;
+    			const float Offset = WeightsAndOffsets[WeightsAndOffsetsIdx + ParamIdx].y;
+
+    			Color += GetSample(ChannelIdx, Weight, Offset, UV);
+    		}
+
+    		{
+    			const float Weight = WeightsAndOffsets[WeightsAndOffsetsIdx + ParamIdx].z;
+    			const float Offset = WeightsAndOffsets[WeightsAndOffsetsIdx + ParamIdx].w;
+
+    			Color += GetSample(ChannelIdx, Weight, Offset, UV);
+    		}
+    	}
+    }
+
+	Color.a = 1.0;
+
+	OutputTexture[PixelCoord] = Color;
+}
diff --git a/Engine/Plugins/Experimental/GeometryMask/Shaders/Private/GeometryMaskJFInitCS.usf b/Engine/Plugins/Experimental/GeometryMask/Shaders/Private/GeometryMaskJFInitCS.usf
new file mode 100644
--- /dev/null
+++ b/Engine/Plugins/Experimental/GeometryMask/Shaders/Private/GeometryMaskJFInitCS.usf
@@ -0,0 +1,107 @@
+// Copyright Epic Games, Inc. All Rights Reserved.
+
+#include "/Engine/Public/Platform.ush"
+
+#define TILE_SIZE 8
+#define KERNEL_SIZE 3
+
+uint2 InputDimensions;
+float2 OneOverInputDimensions;
+float2 UVRatioAdjustment;
+uint4 ChannelMap; // ChannelNum to ChannelIdx (ie. 0 -> Y) 
+Texture2D InputTexture;
+SamplerState InputSampler;
+RWBuffer<float4> OutputBuffer;
+
+[numthreads(TILE_SIZE, TILE_SIZE, 1)]
+void MainCS(
+	uint3 DispatchThreadId : SV_DispatchThreadID,
+	uint3 GroupThreadId : SV_GroupThreadID,
+	uint3 GroupId : SV_GroupID)
+{
+	const uint2 PixelCoord = DispatchThreadId.xy;
+	if (PixelCoord.x > InputDimensions.x - 1 || PixelCoord.y > InputDimensions.y - 1)
+	{
+		return;
+	}
+
+	float X = (float)PixelCoord.x + 0.5;
+	float Y = (float)PixelCoord.y + 0.5;
+	const float2 UV = float2(X, Y) * OneOverInputDimensions * UVRatioAdjustment;
+
+	const uint RowLength = InputDimensions.x * NUM_CHANNELS;	
+	const uint PixelCoord1D = (PixelCoord.y * RowLength) + (PixelCoord.x * NUM_CHANNELS);
+
+	OutputBuffer[PixelCoord1D] = float4(1,1,1,1);
+
+	static int SobelX[] =
+	{
+		1, 0, -1,
+		2, 0, -2,
+		1, 0, -1
+	};
+
+	static int SobelY[] =
+	{
+		1,  2,  1,
+		0,  0,  0,
+	   -1, -2, -1
+   };
+
+	static uint2 NeighborOffsetIndices[] =
+	{
+		{-1,  1},  {0,  1},  {1,  1},
+		{-1,  0},  {0,  0},  {1,  0},
+		{-1, -1},  {0, -1},  {1, -1}
+	};
+	
+	float4 Input = InputTexture.Load(int3(PixelCoord, 0));
+	float4 Edges = float4(0,0,0,0);
+	int4 Mask = int4(0,0,0,0);
+
+	[unroll]
+	for (int ChannelIdx = 0; ChannelIdx < NUM_CHANNELS; ++ChannelIdx)
+	{
+		OutputBuffer[PixelCoord1D + ChannelIdx] = float4(0,0,0,0);
+	}
+	
+	// Perform on each of R,G,B,A
+	[unroll]
+	for (int ChannelIdx = 0; ChannelIdx < NUM_CHANNELS; ++ChannelIdx)
+	{
+		float GradientX = 0;
+		float GradientY = 0;
+
+		const uint InputChannelIdx = ChannelMap[ChannelIdx];
+    	
+		[unroll]
+		for (int HorizontalIdx = 0; HorizontalIdx < 3; ++HorizontalIdx)
+		{
+			[unroll]
+			for (int VerticalIdx = 0; VerticalIdx < 3; ++VerticalIdx)
+			{
+				uint NeighborIdx = (VerticalIdx * KERNEL_SIZE) + HorizontalIdx;
+				uint2 NeighborPixelCoord = PixelCoord + NeighborOffsetIndices[NeighborIdx];
+				NeighborPixelCoord = clamp(NeighborPixelCoord, uint2(0,0), uint2(InputDimensions.x, InputDimensions.y));
+
+				// > 0 ensures treats the input as binary
+				const float NeighborSample = InputTexture.Load(int3(NeighborPixelCoord, 0))[InputChannelIdx] > 0;
+
+				GradientX += SobelX[NeighborIdx] * NeighborSample;
+				GradientY += SobelY[NeighborIdx] * NeighborSample;
+			}
+		}
+
+		Edges[ChannelIdx] = max(abs(GradientX), abs(GradientY));
+		Mask[ChannelIdx] = min(1, Edges[ChannelIdx]);
+	}
+
+	// @note: This is specific to JFA requirements
+	// For each channel in each pixel, encode (UV, Edge)
+	[unroll]
+	for (int ChannelIdx = 0; ChannelIdx < NUM_CHANNELS; ++ChannelIdx)
+	{
+		const uint InputChannelIdx = ChannelMap[ChannelIdx];
+		OutputBuffer[PixelCoord1D + ChannelIdx] = float4(UV * Mask[ChannelIdx], Edges[ChannelIdx], Input[InputChannelIdx]);
+	}
+}
diff --git a/Engine/Plugins/Experimental/GeometryMask/Shaders/Private/GeometryMaskJFStepCS.usf b/Engine/Plugins/Experimental/GeometryMask/Shaders/Private/GeometryMaskJFStepCS.usf
new file mode 100644
--- /dev/null
+++ b/Engine/Plugins/Experimental/GeometryMask/Shaders/Private/GeometryMaskJFStepCS.usf
@@ -0,0 +1,107 @@
+// Copyright Epic Games, Inc. All Rights Reserved.
+
+#include "/Engine/Public/Platform.ush"
+
+#define TILE_SIZE 8
+#define KERNEL_SIZE 3
+#define MAX_STEPS 13
+
+uint2 InputDimensions;
+float2 OneOverInputDimensions;
+float2 UVRatioAdjustment;
+uint4 StepSize; // per-channel
+uint4 ChannelMask; // whether to calculate for this step
+Buffer<float4> InputBuffer;
+RWBuffer<float4> OutputBuffer;
+
+[numthreads(TILE_SIZE, TILE_SIZE, 1)]
+void MainCS(
+	uint3 DispatchThreadId : SV_DispatchThreadID,
+	uint3 GroupThreadId : SV_GroupThreadID,
+	uint3 GroupId : SV_GroupID)
+{
+	uint2 PixelCoord = DispatchThreadId.xy;
+	if (PixelCoord.x > InputDimensions.x - 1 || PixelCoord.y > InputDimensions.y - 1)
+	{
+		return;
+	}
+
+	const uint RowLength = InputDimensions.x * NUM_CHANNELS;
+	uint PixelCoord1D = (PixelCoord.y * RowLength) + (PixelCoord.x * NUM_CHANNELS);
+
+	static uint2 NeighborOffsetIndices[] =
+	{
+		{-1,  1},  {0,  1},  {1,  1},
+		{-1,  0},  {0,  0},  {1,  0},
+		{-1, -1},  {0, -1},  {1, -1}
+	};
+
+#if NUM_CHANNELS == 1
+	float MinDistance[NUM_CHANNELS] = { 100000000 };
+	float4 MinSample[NUM_CHANNELS] = { float4(0,0,0,0) };
+#elif NUM_CHANNELS == 2
+	float MinDistance[NUM_CHANNELS] = { 100000000, 100000000 };
+	float4 MinSample[NUM_CHANNELS] = { float4(0,0,0,0), float4(0,0,0,0) };
+#elif NUM_CHANNELS == 3
+	float MinDistance[NUM_CHANNELS] = { 100000000, 100000000, 100000000 };
+	float4 MinSample[NUM_CHANNELS] = { float4(0,0,0,0), float4(0,0,0,0), float4(0,0,0,0) };
+#elif NUM_CHANNELS == 4
+	float MinDistance[NUM_CHANNELS] = { 100000000, 100000000, 100000000, 100000000 };
+	float4 MinSample[NUM_CHANNELS] = { float4(0,0,0,0), float4(0,0,0,0), float4(0,0,0,0), float4(0,0,0,0) };
+#endif
+
+	float X = (float)PixelCoord.x + 0.5;
+	float Y = (float)PixelCoord.y + 0.5;
+	const float2 UV = float2(X, Y) * OneOverInputDimensions * UVRatioAdjustment;
+
+    // Perform on each of R,G,B,A
+    [unroll]
+    for (int ChannelIdx = 0; ChannelIdx < NUM_CHANNELS; ++ChannelIdx)
+    {
+    	float4 Input = InputBuffer[PixelCoord1D + ChannelIdx];
+    	const uint2 StepSizeWithChannel = StepSize[ChannelIdx] * uint2(1, 1);
+
+    	// Check if we should calculate this channel (for this step)
+    	if (ChannelMask[ChannelIdx] == 0)
+    	{
+    		MinSample[ChannelIdx] = Input[ChannelIdx];
+    		continue;
+    	}
+
+    	MinDistance[ChannelIdx] = 1000000000;
+    	MinSample[ChannelIdx] = float4(0,0,0,0);
+
+    	[unroll]
+		for (int HorizontalIdx = 0; HorizontalIdx < 3; ++HorizontalIdx)
+		{
+			[unroll]
+			for (int VerticalIdx = 0; VerticalIdx < 3; ++VerticalIdx)
+			{
+				uint NeighborIdx = (VerticalIdx * KERNEL_SIZE) + HorizontalIdx;
+				uint2 NeighborPixelCoord = PixelCoord + (NeighborOffsetIndices[NeighborIdx] * StepSizeWithChannel);
+				if (NeighborPixelCoord.x > InputDimensions.x - 1 || NeighborPixelCoord.y > InputDimensions.y - 1)
+				{
+					continue;
+				}
+
+				uint NeighborPixelCoord1D = (NeighborPixelCoord.y * RowLength) + (NeighborPixelCoord.x * NUM_CHANNELS);
+
+				float4 NeighborSample = InputBuffer[NeighborPixelCoord1D + ChannelIdx];
+				if (NeighborSample.x + NeighborSample.y == 0.0)
+				{
+					continue;
+				}
+
+				const float Distance = length(NeighborSample.xy - UV);
+				if (Distance < MinDistance[ChannelIdx])
+				{
+					MinDistance[ChannelIdx] = Distance;
+					MinSample[ChannelIdx] = NeighborSample;
+				}
+			}
+		}
+
+    	MinSample[ChannelIdx].a = Input.a;
+    	OutputBuffer[PixelCoord1D + ChannelIdx] = MinSample[ChannelIdx];
+    }
+}
diff --git a/Engine/Plugins/Experimental/GeometryMask/Shaders/Private/GeometryMaskJFtoDFCS.usf b/Engine/Plugins/Experimental/GeometryMask/Shaders/Private/GeometryMaskJFtoDFCS.usf
new file mode 100644
--- /dev/null
+++ b/Engine/Plugins/Experimental/GeometryMask/Shaders/Private/GeometryMaskJFtoDFCS.usf
@@ -0,0 +1,58 @@
+// Copyright Epic Games, Inc. All Rights Reserved.
+
+#include "/Engine/Public/Platform.ush"
+
+#define TILE_SIZE 8
+
+uint2 InputDimensions;
+float2 OneOverInputDimensions;
+float2 UVRatioAdjustment;
+uint4 ChannelMap; // ChannelNum to ChannelIdx (ie. 0 -> Y)
+float4 StepDistanceMultipliers; // To normalize
+Texture2D OriginalInputTexture;
+SamplerState OriginalInputSampler;
+Buffer<float4> InputBuffer;
+RWTexture2D<float4> OutputTexture;
+
+[numthreads(TILE_SIZE, TILE_SIZE, 1)]
+void MainCS(
+	uint3 DispatchThreadId : SV_DispatchThreadID,
+	uint3 GroupThreadId : SV_GroupThreadID,
+	uint3 GroupId : SV_GroupID)
+{
+	uint2 PixelCoord = DispatchThreadId.xy;
+	if (PixelCoord.x > InputDimensions.x - 1 || PixelCoord.y > InputDimensions.y - 1)
+	{
+		return;
+	}
+
+	const uint RowLength = InputDimensions.x * NUM_CHANNELS;	
+	const uint PixelCoord1D = (PixelCoord.y * RowLength) + (PixelCoord.x * NUM_CHANNELS);
+
+	float X = (float)PixelCoord.x + 0.5;
+	float Y = (float)PixelCoord.y + 0.5;
+	const float2 UV = float2(X, Y) * OneOverInputDimensions * UVRatioAdjustment;
+
+	float4 Outline = OriginalInputTexture.Load(int3(PixelCoord, 0));
+	float4 Mask = float4(1,1,1,1);
+
+	// Perform on each of R,G,B,A
+	[unroll]
+	for (int ChannelIdx = 0; ChannelIdx < NUM_CHANNELS; ++ChannelIdx)
+	{
+		const uint OutputChannelIdx = ChannelMap[ChannelIdx];
+		
+		float4 Sample = InputBuffer[PixelCoord1D + ChannelIdx];
+		Mask[OutputChannelIdx] = (Sample.x + Sample.y + Sample.a > 0.001);
+
+		Outline[OutputChannelIdx] = saturate(distance(UV, Sample.xy) * StepDistanceMultipliers[OutputChannelIdx]);
+
+		// Set range such that inner going from 0.0-0.5, outer goes from 0.5-1.0
+		Outline[OutputChannelIdx] = mad(Outline[OutputChannelIdx], 0.5, 0.5);
+		Outline[OutputChannelIdx] = 1.0 - abs(Outline[OutputChannelIdx] - saturate(Sample.a));
+	}
+
+	Outline = lerp(0, Outline, Mask);
+
+	OutputTexture[PixelCoord] = Outline;
+}
diff --git a/Engine/Plugins/Experimental/GeometryMask/Source/GeometryMask/GeometryMask.Build.cs b/Engine/Plugins/Experimental/GeometryMask/Source/GeometryMask/GeometryMask.Build.cs
new file mode 100644
--- /dev/null
+++ b/Engine/Plugins/Experimental/GeometryMask/Source/GeometryMask/GeometryMask.Build.cs
@@ -0,0 +1,42 @@
+﻿// Copyright Epic Games, Inc. All Rights Reserved.
+
+using System.IO;
+using UnrealBuildTool;
+
+public class GeometryMask : ModuleRules
+{
+	public GeometryMask(ReadOnlyTargetRules Target) : base(Target)
+	{
+		PCHUsage = ModuleRules.PCHUsageMode.UseExplicitOrSharedPCHs;
+
+		PublicDependencyModuleNames.AddRange(
+			new string[]
+			{
+				"Core",
+				"CoreUObject",
+				"Engine",
+			});
+
+		PrivateDependencyModuleNames.AddRange(
+			new string[]
+			{
+				"DeveloperSettings",
+				"GeometryCore",
+				"GeometryFramework",
+				"Projects",
+				"RHI",
+				"RenderCore",
+				"Renderer",
+			});
+
+		if (Target.Type == TargetType.Editor)
+		{
+			PrivateDependencyModuleNames.AddRange(
+				new string[]
+				{
+					"LevelEditor",
+					"UnrealEd",
+				});
+		}
+	}
+}
diff --git a/Engine/Plugins/Experimental/GeometryMask/Source/GeometryMask/Private/GeometryMaskCanvas.cpp b/Engine/Plugins/Experimental/GeometryMask/Source/GeometryMask/Private/GeometryMaskCanvas.cpp
new file mode 100644
--- /dev/null
+++ b/Engine/Plugins/Experimental/GeometryMask/Source/GeometryMask/Private/GeometryMaskCanvas.cpp
@@ -0,0 +1,239 @@
+﻿// Copyright Epic Games, Inc. All Rights Reserved.
+
+#include "GeometryMaskCanvas.h"
+
+#include "Algo/RemoveIf.h"
+#include "CanvasTypes.h"
+#include "Engine/Engine.h"
+#include "Engine/TextureRenderTarget2D.h"
+#include "Engine/World.h"
+#include "GeometryMaskCanvasResource.h"
+#include "IGeometryMaskWriteInterface.h"
+
+void UGeometryMaskCanvas::BeginDestroy()
+{
+	FreeResource();
+	
+	UObject::BeginDestroy();
+}
+
+const TArray<TWeakInterfacePtr<IGeometryMaskWriteInterface>>& UGeometryMaskCanvas::GetWriters() const
+{
+	return Writers;
+}
+
+void UGeometryMaskCanvas::AddWriter(const TScriptInterface<IGeometryMaskWriteInterface>& InWriter)
+{
+	if (!ensure(InWriter))
+	{
+		return;
+	}
+
+	RemoveInvalidWriters();
+
+	// Writers is empty, but about to become activated
+	if (Writers.IsEmpty())
+	{
+		OnActivated().ExecuteIfBound();
+	}
+
+	Writers.Add(InWriter.GetObject());
+}
+
+void UGeometryMaskCanvas::AddWriters(const TArray<TScriptInterface<IGeometryMaskWriteInterface>>& InWriters)
+{
+	if (InWriters.IsEmpty())
+	{
+		return;
+	}
+
+	RemoveInvalidWriters();
+	for (const TScriptInterface<IGeometryMaskWriteInterface>& Writer : InWriters)
+	{
+		AddWriter(Writer);		
+	}
+}
+
+void UGeometryMaskCanvas::RemoveWriter(const TScriptInterface<IGeometryMaskWriteInterface>& InWriter)
+{
+	if (!ensure(InWriter))
+	{
+		return;
+	}
+
+	RemoveInvalidWriters();
+	Writers.Remove(InWriter.GetObject());
+}
+
+UTextureRenderTarget2D* UGeometryMaskCanvas::GetTexture() const
+{
+	if (CanvasResource)
+	{
+		return CanvasResource->GetRenderTargetTexture();
+	}
+
+	return nullptr;
+}
+
+bool UGeometryMaskCanvas::IsBlurApplied() const
+{
+	return bApplyBlur;
+}
+
+void UGeometryMaskCanvas::SetApplyBlur(const bool bInValue)
+{
+	if (bInValue != bApplyBlur)
+	{
+		bApplyBlur = bInValue;
+		UpdateRenderParameters();
+	}
+}
+
+double UGeometryMaskCanvas::GetBlurStrength() const
+{
+	return BlurStrength;
+}
+
+void UGeometryMaskCanvas::SetBlurStrength(const double InValue)
+{
+	if (!FMath::IsNearlyEqual(InValue, BlurStrength))
+	{
+		BlurStrength = InValue;
+		UpdateRenderParameters();
+	}
+}
+
+bool UGeometryMaskCanvas::IsFeatherApplied() const
+{
+	return bApplyFeather;
+}
+
+void UGeometryMaskCanvas::SetApplyFeather(const bool bInValue)
+{
+	if (bInValue != bApplyFeather)
+	{
+		bApplyFeather = bInValue;
+		UpdateRenderParameters();
+	}
+}
+
+int32 UGeometryMaskCanvas::GetOuterFeatherRadius() const
+{
+	return OuterFeatherRadius;
+}
+
+void UGeometryMaskCanvas::SetOuterFeatherRadius(const int32 InValue)
+{
+	if (InValue != OuterFeatherRadius)
+	{
+		OuterFeatherRadius = InValue;
+		UpdateRenderParameters();
+	}
+}
+
+int32 UGeometryMaskCanvas::GetInnerFeatherRadius() const
+{
+	return InnerFeatherRadius;
+}
+
+void UGeometryMaskCanvas::SetInnerFeatherRadius(const int32 InValue)
+{
+	if (InValue != InnerFeatherRadius)
+	{
+		InnerFeatherRadius = InValue;
+		UpdateRenderParameters();
+	}
+}
+
+EGeometryMaskColorChannel UGeometryMaskCanvas::GetColorChannel() const
+{
+	return ColorChannel;
+}
+
+void UGeometryMaskCanvas::RemoveInvalidWriters()
+{
+	if (Writers.IsEmpty())
+	{
+		return;
+	}
+
+	Writers.SetNum(Algo::RemoveIf(Writers, [](const TWeakInterfacePtr<IGeometryMaskWriteInterface>& InWriter)
+	{
+		return !InWriter.IsValid() || !IsValid(InWriter.GetObject());
+	}));
+
+	if (Writers.IsEmpty())
+	{
+		// No writers after cleanup, flag as deactivated
+		OnDeactivated().ExecuteIfBound();
+	}
+}
+
+void UGeometryMaskCanvas::OnDrawToCanvas(FCanvas* InCanvas)
+{
+	// Sort so subtract happens after additive, etc.
+	SortWriters();
+
+	for (const TWeakInterfacePtr<IGeometryMaskWriteInterface>& Writer : Writers)
+	{
+		if (Writer.IsValid())
+		{
+			TRACE_CPUPROFILER_EVENT_SCOPE(UGeometryMaskCanvas::OnDrawToCanvas);
+			Writer->DrawToCanvas(InCanvas);	
+		}
+	}
+}
+
+void UGeometryMaskCanvas::UpdateRenderParameters()
+{
+	if (CanvasResource)
+	{
+		CanvasResource->UpdateRenderParameters(ColorChannel, bApplyBlur, BlurStrength, bApplyFeather, OuterFeatherRadius, InnerFeatherRadius);
+	}
+}
+
+void UGeometryMaskCanvas::Update(
+	UWorld* InWorld,
+	FSceneView& InView)
+{
+	RemoveInvalidWriters();
+}
+
+void UGeometryMaskCanvas::AssignResource(
+	UGeometryMaskCanvasResource* InResource,
+	const EGeometryMaskColorChannel InColorChannel)
+{
+	// Cleanup existing resource, if assigned
+	FreeResource();
+
+	CanvasResource = InResource;
+	ColorChannel = InColorChannel;
+
+	CanvasResource->OnDrawToCanvas().AddUObject(this, &UGeometryMaskCanvas::OnDrawToCanvas);
+}
+
+void UGeometryMaskCanvas::FreeResource()
+{
+	if (CanvasResource)
+	{
+		CanvasResource->OnDrawToCanvas().RemoveAll(this);
+		ColorChannel = EGeometryMaskColorChannel::None;
+	
+		// Free up the resource
+		CanvasResource->Checkin(CanvasName);
+		CanvasResource = nullptr;
+	}
+}
+
+void UGeometryMaskCanvas::SortWriters()
+{
+	Algo::SortBy(Writers, [](const TWeakInterfacePtr<IGeometryMaskWriteInterface>& InWriter)
+	{
+		if (const IGeometryMaskWriteInterface* Writer = InWriter.Get())
+		{
+			return Writer->GetParameters().OperationType;
+		}
+
+		return EGeometryMaskCompositeOperation::Num;
+	});
+}
diff --git a/Engine/Plugins/Experimental/GeometryMask/Source/GeometryMask/Private/GeometryMaskCanvasActor.cpp b/Engine/Plugins/Experimental/GeometryMask/Source/GeometryMask/Private/GeometryMaskCanvasActor.cpp
new file mode 100644
--- /dev/null
+++ b/Engine/Plugins/Experimental/GeometryMask/Source/GeometryMask/Private/GeometryMaskCanvasActor.cpp
@@ -0,0 +1,101 @@
+﻿// Copyright Epic Games, Inc. All Rights Reserved.
+
+#include "GeometryMaskCanvasActor.h"
+
+#include "Engine/Engine.h"
+#include "GeometryMaskCanvas.h"
+#include "GeometryMaskSubsystem.h"
+#include "GeometryMaskWriteComponent.h"
+
+AGeometryMaskCanvasActor::AGeometryMaskCanvasActor(const FObjectInitializer& ObjectInitializer)
+	: Super(ObjectInitializer)
+{
+	PrimaryActorTick.bCanEverTick = true;
+	PrimaryActorTick.bStartWithTickEnabled = true;
+	SetActorTickEnabled(true);
+
+#if WITH_EDITORONLY_DATA
+	bIsSpatiallyLoaded = false;
+#endif
+
+	RootComponent = CreateDefaultSubobject<USceneComponent>(TEXT("Root"));
+
+	Canvas = nullptr;
+}
+
+UTextureRenderTarget2D* AGeometryMaskCanvasActor::GetTexture()
+{
+	return Canvas ? Canvas->GetTexture() : nullptr;
+}
+
+void AGeometryMaskCanvasActor::BeginPlay()
+{
+	Super::BeginPlay();
+
+	TryResolveCanvas();
+}
+
+void AGeometryMaskCanvasActor::PostLoad()
+{
+	Super::PostLoad();
+
+	if (!HasAnyFlags(RF_ClassDefaultObject))
+	{
+		TryResolveCanvas();
+	}
+}
+
+#if WITH_EDITOR
+void AGeometryMaskCanvasActor::RerunConstructionScripts()
+{
+	Super::RerunConstructionScripts();
+
+	if (!HasAnyFlags(RF_ClassDefaultObject))
+	{
+		TryResolveCanvas();
+	}
+}
+#endif
+
+bool AGeometryMaskCanvasActor::TryResolveCanvas()
+{
+	if (!Canvas)
+	{
+		if (UGeometryMaskSubsystem* Subsystem = GEngine->GetEngineSubsystem<UGeometryMaskSubsystem>())
+		{
+			Canvas = Subsystem->GetNamedCanvas(CanvasName);
+		}
+	}
+	
+	FindWriters();
+	
+	return Canvas != nullptr;
+}
+
+void AGeometryMaskCanvasActor::FindWriters()
+{
+	// Find actors that implement, or contain components that implement IGeometryMaskWriteInterface
+	TArray<AActor*> ChildActors;
+	GetAttachedActors(ChildActors, true, true);
+
+	TArray<TScriptInterface<IGeometryMaskWriteInterface>> FoundWriters;
+	FoundWriters.Reserve(ChildActors.Num());
+
+	for (const TObjectPtr<AActor> ChildActor : ChildActors)
+	{
+		if (ChildActor->Implements<UGeometryMaskWriteInterface>())
+		{
+			FoundWriters.Add(ChildActor);
+		}
+
+		FoundWriters.Append(ChildActor->GetComponentsByInterface(UGeometryMaskWriteInterface::StaticClass()));
+	}
+
+	Writers.Reset();
+	Writers = FoundWriters;
+
+	if (Canvas)
+	{
+		Canvas->AddWriters(Writers);
+	}
+}
diff --git a/Engine/Plugins/Experimental/GeometryMask/Source/GeometryMask/Private/GeometryMaskCanvasResource.cpp b/Engine/Plugins/Experimental/GeometryMask/Source/GeometryMask/Private/GeometryMaskCanvasResource.cpp
new file mode 100644
--- /dev/null
+++ b/Engine/Plugins/Experimental/GeometryMask/Source/GeometryMask/Private/GeometryMaskCanvasResource.cpp
@@ -0,0 +1,355 @@
+// Copyright Epic Games, Inc. All Rights Reserved.
+
+#include "GeometryMaskCanvasResource.h"
+
+#include "Engine/Canvas.h"
+#include "Engine/TextureRenderTarget2D.h"
+#include "Engine/World.h"
+#include "GeometryMaskModule.h"
+#include "GeometryMaskSettings.h"
+#include "GeometryMaskTypes.h"
+#include "SceneView.h"
+#include "Shaders/GeometryMaskPostProcess.h"
+#include "Shaders/GeometryMaskPostProcess_Blur.h"
+#include "Shaders/GeometryMaskPostProcess_DistanceField.h"
+#include "TextureResource.h"
+#include "UObject/Package.h"
+
+namespace UE::GeometryMask::Private
+{
+	void OverscanProjectionMatrix(FMatrix& InOutMatrix, const FIntPoint& InSize, const int32 InPadding)
+	{
+		const FVector2f Multiplier(
+			static_cast<float>(InSize.X) / static_cast<float>(InSize.X + InPadding),
+			static_cast<float>(InSize.Y) / static_cast<float>(InSize.Y + InPadding));
+
+		// Original Calc - we simply scale each
+		// [0][0] = MultFOVX / FMath::Tan(HalfFOVX)
+		// [1][1] = MultFOVY / FMath::Tan(HalfFOVY)
+		
+		InOutMatrix.M[0][0] *= Multiplier.X;
+		InOutMatrix.M[1][1] *= Multiplier.Y;
+	}
+}
+
+UGeometryMaskCanvasResource::UGeometryMaskCanvasResource()
+{
+	// @note: We omit the Alpha channel as it's not uniformly supported
+	DependentCanvasNames = {
+		{ EGeometryMaskColorChannel::Red, NAME_None },
+		{ EGeometryMaskColorChannel::Green, NAME_None },
+		{ EGeometryMaskColorChannel::Blue, NAME_None }
+	};
+
+	FGeometryMaskPostProcessParameters_Blur PostProcessParameters_Blur;
+	PostProcessParameters_Blur.PerChannelBlurStrength = { 16, 16, 16, 16};
+	PostProcess_Blur = MakeShared<FGeometryMaskPostProcess_Blur>(PostProcessParameters_Blur);
+
+	FGeometryMaskPostProcessParameters_DistanceField PostProcessParameters_DistanceField;
+	PostProcessParameters_DistanceField.bPerChannelCalculateDF.SetRange(0, 4, true);
+	PostProcess_DistanceField = MakeShared<FGeometryMaskPostProcess_DistanceField>(PostProcessParameters_DistanceField);
+}
+
+const EGeometryMaskColorChannel UGeometryMaskCanvasResource::GetNextAvailableColorChannel() const
+{
+	for (const TPair<EGeometryMaskColorChannel, FName>& ColorChannelCanvas : DependentCanvasNames)
+	{
+		if (ColorChannelCanvas.Value.IsNone())
+		{
+			return ColorChannelCanvas.Key;
+		}
+	}
+
+	return EGeometryMaskColorChannel::None;
+}
+
+bool UGeometryMaskCanvasResource::Checkout(
+	const EGeometryMaskColorChannel InColorChannel
+	, const FName InRequestingCanvasName)
+{
+	if (!ensure(InColorChannel != EGeometryMaskColorChannel::Num && InColorChannel != EGeometryMaskColorChannel::None))
+	{
+		return false;
+	}
+
+	if (!ensure(DependentCanvasNames[InColorChannel].IsNone()))
+	{
+		return false;
+	}
+
+	DependentCanvasNames[InColorChannel] = InRequestingCanvasName;
+	return true;
+}
+
+bool UGeometryMaskCanvasResource::Checkin(const FName InRequestingCanvasName)
+{
+	for (TPair<EGeometryMaskColorChannel, FName>& ColorChannelCanvas : DependentCanvasNames)
+	{
+		if (ColorChannelCanvas.Value.IsEqual(InRequestingCanvasName))
+		{
+			// Effectively free this ColorChannel and make available for Checkout
+			ColorChannelCanvas.Value = NAME_None;
+			return true;
+		}
+	}
+
+	// The requesting canvas name wasn't present in this Resource
+	return false;
+}
+
+void UGeometryMaskCanvasResource::UpdateViewportSize()
+{
+	if (ViewportSize.Size() > 0
+		&& RenderTargetTexture)
+	{
+		const float SizeMultiplier = GetDefault<UGeometryMaskSettings>()->GetDefaultResolutionMultiplier();
+		
+		int32 SizeX = ViewportSize.X * SizeMultiplier;
+		int32 SizeY = ViewportSize.Y * SizeMultiplier;
+
+		const int32 Padding = GetViewportPadding();
+		SizeX += Padding * SizeMultiplier;
+		SizeY += Padding * SizeMultiplier;
+
+		if (const float RatioX = static_cast<float>(SizeX) / MaxTextureSize;
+			RatioX > 1.0)
+		{
+			// Width too big, cap to max and reduce height proportionally
+			SizeX = MaxTextureSize;
+			SizeY /= RatioX;
+		}
+
+		if (const float RatioY = static_cast<float>(SizeY) / MaxTextureSize;
+			RatioY > 1.0)
+		{
+			// Height too big, cap to max and reduce width proportionally
+			SizeY = MaxTextureSize;
+			SizeX /= RatioY;
+		}
+
+		// Update RT size to viewport size
+		RenderTargetTexture->ResizeTarget(SizeX, SizeY);
+	}
+}
+
+void UGeometryMaskCanvasResource::UpdateRenderParameters(
+	const EGeometryMaskColorChannel InColorChannel,
+	const bool bInApplyBlur,
+	const double InBlurStrength,
+	bool bInApplyFeather,
+	int32 InOuterFeatherRadius,
+	int32 InInnerFeatherRadius)
+{
+	const int32 ChannelIdx = static_cast<int32>(InColorChannel);
+
+	auto LogOutOfBounds = [InColorChannel](){
+		const FStringView ChannelStringView = UE::GeometryMask::ChannelToString(InColorChannel);
+		TCHAR* ChannelString = nullptr;
+		ChannelStringView.CopyString(ChannelString, ChannelStringView.Len());
+		
+		UE_LOG(LogGeometryMask, Error, TEXT("ColorChannel wasn't valid for setting shader parameters (was '%s', expected R, G, B or A"), ChannelString);
+	};
+
+	FGeometryMaskPostProcessParameters_Blur BlurParameters = PostProcess_Blur->GetParameters();
+	{
+		if (BlurParameters.bPerChannelApplyBlur.IsValidIndex(ChannelIdx))
+		{
+			BlurParameters.bPerChannelApplyBlur[ChannelIdx] = bInApplyBlur && InBlurStrength > 0.0;
+		}
+		else
+		{
+			LogOutOfBounds();
+			return;
+		}
+
+		if (BlurParameters.PerChannelBlurStrength.IsValidIndex(ChannelIdx))
+		{
+			BlurParameters.PerChannelBlurStrength[ChannelIdx] = InBlurStrength;	
+		}
+		else
+		{
+			LogOutOfBounds();
+			return;
+		}
+
+		PostProcess_Blur->SetParameters(BlurParameters);
+	}
+
+	FGeometryMaskPostProcessParameters_DistanceField DFParameters = PostProcess_DistanceField->GetParameters();
+	{
+		if (DFParameters.bPerChannelCalculateDF.IsValidIndex(ChannelIdx))
+		{
+			DFParameters.bPerChannelCalculateDF[ChannelIdx] = bInApplyFeather && (InOuterFeatherRadius + InInnerFeatherRadius) > 0;
+		}
+		else
+		{
+			LogOutOfBounds();
+			return;
+		}
+
+		if (DFParameters.PerChannelRadius.IsValidIndex(ChannelIdx))
+		{
+			DFParameters.PerChannelRadius[ChannelIdx] = FMath::Max(InOuterFeatherRadius, InInnerFeatherRadius);	
+		}
+		else
+		{
+			LogOutOfBounds();
+			return;
+		}
+
+		PostProcess_DistanceField->SetParameters(DFParameters);
+	}
+
+	// Effects may have changed viewport padding
+	UpdateViewportSize();
+
+	bApplyBlur = BlurParameters.bPerChannelApplyBlur.CountSetBits(0) > 0; // if any channels have blur
+	bApplyDF = DFParameters.bPerChannelCalculateDF.CountSetBits(0) > 0; // if any channels have DF
+}
+
+void UGeometryMaskCanvasResource::SetViewportSize(const FIntPoint& InViewportSize)
+{
+	if (InViewportSize.Size() > 0
+		&& ViewportSize != InViewportSize)
+	{
+		ViewportSize = InViewportSize;
+		UpdateViewportSize();		
+	}
+}
+
+int32 UGeometryMaskCanvasResource::GetViewportPadding() const
+{
+	int32 MaxFeatherRadius = 0;
+	{
+		const FGeometryMaskPostProcessParameters_DistanceField& DistanceFieldParameters = PostProcess_DistanceField->GetParameters();
+		for (int32 ChannelIdx = 0; ChannelIdx < MaxNumChannels; ++ChannelIdx)
+		{
+			if (DistanceFieldParameters.bPerChannelCalculateDF[ChannelIdx])
+			{
+				MaxFeatherRadius = FMath::Max(MaxFeatherRadius, DistanceFieldParameters.PerChannelRadius[ChannelIdx]);
+			}
+		}
+	}
+
+	int32 MaxBlurRadius = 0;
+	{
+		const FGeometryMaskPostProcessParameters_Blur& BlurParameters = PostProcess_Blur->GetParameters();
+		for (int32 ChannelIdx = 0; ChannelIdx < MaxNumChannels; ++ChannelIdx)
+		{
+			if (BlurParameters.bPerChannelApplyBlur[ChannelIdx])
+			{
+				MaxBlurRadius = FMath::Max(MaxBlurRadius, UE::GeometryMask::Internal::ComputeEffectiveKernelSize(BlurParameters.PerChannelBlurStrength[ChannelIdx]));
+			}
+		}
+	}
+
+	
+	return FMath::Max(MaxFeatherRadius, MaxBlurRadius);
+}
+
+UTextureRenderTarget2D* UGeometryMaskCanvasResource::GetRenderTargetTexture()
+{
+	if (!RenderTargetTexture)
+	{
+		const FName ObjectName = MakeUniqueObjectName(this, UTextureRenderTarget2D::StaticClass(), FName(FString::Printf(TEXT("GeometryMaskCanvasResource_RenderTarget"))));
+		RenderTargetTexture = NewObject<UTextureRenderTarget2D>(this, ObjectName);
+		RenderTargetTexture->bForceLinearGamma = false;
+		RenderTargetTexture->InitAutoFormat(ViewportSize.X, ViewportSize.Y);
+	}
+
+	return RenderTargetTexture;
+}
+
+void UGeometryMaskCanvasResource::Update(
+	UWorld* InWorld,
+	FSceneView& InView)
+{
+	Draw(InWorld, InView);
+}
+
+void UGeometryMaskCanvasResource::Draw(UWorld* InWorld, FSceneView& InView)
+{
+	if (!InWorld)
+	{
+		return;
+	}
+
+	if (UTextureRenderTarget2D* Texture = GetRenderTargetTexture())
+	{
+		if (!CanvasObject)
+		{
+			CanvasObject = NewObject<UCanvas>(GetTransientPackage());
+		}
+
+		// Begin
+		{
+			SetViewportSize(InView.UnconstrainedViewRect.Size());
+			FMatrix ProjectionMatrix = InView.ViewMatrices.GetProjectionMatrix();
+			UE::GeometryMask::Private::OverscanProjectionMatrix(ProjectionMatrix, ViewportSize, GetViewportPadding());
+
+			CachedViewProjectionMatrix = InView.ViewMatrices.GetViewMatrix() * ProjectionMatrix;
+			CachedViewProjectionMatrix.M[2][2] = UE_KINDA_SMALL_NUMBER; // Prevents div by zero later
+
+			UpdateViewportSize();
+			
+			InWorld->FlushDeferredParameterCollectionInstanceUpdates();
+
+			FTextureRenderTargetResource* RenderTargetResource = Texture->GameThread_GetRenderTargetResource();
+			FCanvas* NewCanvas = new FCanvas(
+				RenderTargetResource,
+				nullptr,
+				InWorld,
+				InWorld->GetFeatureLevel(),
+				// Draw immediately so that interleaved SetVectorParameter (etc) function calls work as expected
+				FCanvas::CDM_ImmediateDrawing);
+
+			CanvasObject->Init(Texture->SizeX, Texture->SizeY, &InView, NewCanvas);
+			CanvasObject->Update();
+			CanvasObject->SetView(&InView);
+			CanvasObject->Canvas->Clear(GetRenderTargetTexture()->ClearColor);
+		}
+
+		// Contents
+		{
+			// Store current transform
+			const FMatrix CanvasMatrix = CanvasObject->Canvas->GetBottomTransform();
+
+			// Set to World->Viewport transform
+			CanvasObject->Canvas->SetBaseTransform(CachedViewProjectionMatrix);
+			{
+				TRACE_CPUPROFILER_EVENT_SCOPE(UGeometryMaskCanvasResource::Draw::OnDrawToCanvas);
+				
+				OnDrawToCanvas().Broadcast(CanvasObject->Canvas);
+			}
+
+			// Restore original transform
+			CanvasObject->Canvas->SetBaseTransform(CanvasMatrix);
+		}
+		
+		// End
+		{
+			if (CanvasObject && CanvasObject->Canvas)
+			{
+				CanvasObject->Canvas->Flush_GameThread();
+				
+				// Post Process
+				if (bApplyDF || bApplyBlur)
+				{
+					FRenderTarget* RenderTarget = CanvasObject->Canvas->GetRenderTarget();
+					if (bApplyDF)
+					{
+						PostProcess_DistanceField->Execute(RenderTarget);	
+					}
+
+					if (bApplyBlur)
+					{
+						PostProcess_Blur->Execute(RenderTarget);
+					}
+				}
+
+				delete CanvasObject->Canvas;
+				CanvasObject->Canvas = nullptr;
+			}
+		}
+	}
+}
diff --git a/Engine/Plugins/Experimental/GeometryMask/Source/GeometryMask/Private/GeometryMaskModule.cpp b/Engine/Plugins/Experimental/GeometryMask/Source/GeometryMask/Private/GeometryMaskModule.cpp
new file mode 100644
--- /dev/null
+++ b/Engine/Plugins/Experimental/GeometryMask/Source/GeometryMask/Private/GeometryMaskModule.cpp
@@ -0,0 +1,29 @@
+﻿// Copyright Epic Games, Inc. All Rights Reserved.
+
+#include "GeometryMaskModule.h"
+
+#include "IGeometryMaskModule.h"
+#include "Interfaces/IPluginManager.h"
+#include "Misc/Paths.h"
+#include "Modules/ModuleManager.h"
+#include "ShaderCore.h"
+
+DEFINE_LOG_CATEGORY(LogGeometryMask);
+
+#define LOCTEXT_NAMESPACE "FGeometryMaskModule"
+
+IGeometryMaskModule& IGeometryMaskModule::Get()
+{
+	return FModuleManager::LoadModuleChecked<IGeometryMaskModule>(UE_MODULE_NAME);
+}
+
+void FGeometryMaskModule::StartupModule()
+{
+	const FString PluginShaderDir = FPaths::Combine(IPluginManager::Get().FindPlugin(UE_PLUGIN_NAME)->GetBaseDir(), TEXT("Shaders"));
+	// @todo: uncomment when shaders implemented
+	AddShaderSourceDirectoryMapping(TEXT("/Plugin/GeometryMask"), PluginShaderDir);
+}
+
+#undef LOCTEXT_NAMESPACE
+
+IMPLEMENT_MODULE(FGeometryMaskModule, GeometryMask)
diff --git a/Engine/Plugins/Experimental/GeometryMask/Source/GeometryMask/Private/GeometryMaskModule.h b/Engine/Plugins/Experimental/GeometryMask/Source/GeometryMask/Private/GeometryMaskModule.h
new file mode 100644
--- /dev/null
+++ b/Engine/Plugins/Experimental/GeometryMask/Source/GeometryMask/Private/GeometryMaskModule.h
@@ -0,0 +1,19 @@
+﻿// Copyright Epic Games, Inc. All Rights Reserved.
+
+#pragma once
+
+#include "IGeometryMaskModule.h"
+
+#include "Logging/LogMacros.h"
+#include "Modules/ModuleInterface.h"
+
+DECLARE_LOG_CATEGORY_EXTERN(LogGeometryMask, Log, All);
+
+class FGeometryMaskModule
+	: public IGeometryMaskModule
+{
+public:
+	// ~Begin IModuleInterface
+	virtual void StartupModule() override;
+	// ~End IModuleInterface
+};
diff --git a/Engine/Plugins/Experimental/GeometryMask/Source/GeometryMask/Private/GeometryMaskReadComponent.cpp b/Engine/Plugins/Experimental/GeometryMask/Source/GeometryMask/Private/GeometryMaskReadComponent.cpp
new file mode 100644
--- /dev/null
+++ b/Engine/Plugins/Experimental/GeometryMask/Source/GeometryMask/Private/GeometryMaskReadComponent.cpp
@@ -0,0 +1,22 @@
+﻿// Copyright Epic Games, Inc. All Rights Reserved.
+
+#include "GeometryMaskReadComponent.h"
+
+#include "GeometryMaskCanvas.h"
+
+void UGeometryMaskReadComponent::SetParameters(FGeometryMaskReadParameters& InParameters)
+{
+	Parameters = InParameters;
+	TryResolveCanvas();
+}
+
+bool UGeometryMaskReadComponent::TryResolveCanvas()
+{
+	const bool bResult = TryResolveNamedCanvas(Parameters.CanvasName);
+	if (bResult)
+	{
+		Parameters.ColorChannel = CanvasWeak.Get()->GetColorChannel();
+	}
+	
+	return bResult;
+}
diff --git a/Engine/Plugins/Experimental/GeometryMask/Source/GeometryMask/Private/GeometryMaskSVE.cpp b/Engine/Plugins/Experimental/GeometryMask/Source/GeometryMask/Private/GeometryMaskSVE.cpp
new file mode 100644
--- /dev/null
+++ b/Engine/Plugins/Experimental/GeometryMask/Source/GeometryMask/Private/GeometryMaskSVE.cpp
@@ -0,0 +1,54 @@
+﻿// Copyright Epic Games, Inc. All Rights Reserved.
+
+#include "GeometryMaskSVE.h"
+
+#include "GeometryMaskSubsystem.h"
+
+#if WITH_EDITOR
+#include "Editor.h"
+#endif
+
+FGeometryMaskSceneViewExtension::FGeometryMaskSceneViewExtension(
+	const FAutoRegister& AutoRegister,
+	UWorld* InWorld)
+	: FWorldSceneViewExtension(AutoRegister, InWorld)
+{
+}
+
+void FGeometryMaskSceneViewExtension::BeginRenderViewFamily(
+	FSceneViewFamily& InViewFamily)
+{
+	if (UGeometryMaskSubsystem* Subsystem = GEngine->GetEngineSubsystem<UGeometryMaskSubsystem>())
+	{
+		Subsystem->Update(GetWorld(), InViewFamily);
+	}
+}
+
+bool FGeometryMaskSceneViewExtension::IsActiveThisFrame_Internal(const FSceneViewExtensionContext& Context) const
+{
+	bool bIsActive = FWorldSceneViewExtension::IsActiveThisFrame_Internal(Context);
+	if (!bIsActive)
+	{
+		return false;
+	}
+
+#if WITH_EDITOR
+	if (GEditor)
+	{
+		if (GEditor->IsSimulatingInEditor())
+		{
+			bIsActive = GetWorld()->WorldType == EWorldType::Editor;
+		}
+		else if (GEditor->PlayWorld)
+		{
+			bIsActive = GetWorld()->WorldType == EWorldType::PIE;
+		}
+		else
+		{
+			bIsActive = GetWorld()->WorldType == EWorldType::Editor;	
+		}		
+	}
+#endif
+
+	return bIsActive;
+}
diff --git a/Engine/Plugins/Experimental/GeometryMask/Source/GeometryMask/Private/GeometryMaskSVE.h b/Engine/Plugins/Experimental/GeometryMask/Source/GeometryMask/Private/GeometryMaskSVE.h
new file mode 100644
--- /dev/null
+++ b/Engine/Plugins/Experimental/GeometryMask/Source/GeometryMask/Private/GeometryMaskSVE.h
@@ -0,0 +1,21 @@
+﻿// Copyright Epic Games, Inc. All Rights Reserved.
+
+#pragma once
+
+#include "SceneViewExtension.h"
+
+class UGeometryMaskCanvas;
+
+class FGeometryMaskSceneViewExtension
+	: public FWorldSceneViewExtension
+{
+public:
+	FGeometryMaskSceneViewExtension(const FAutoRegister& AutoRegister, UWorld* InWorld);
+
+	// ~Begin ISceneViewExtension
+	virtual void SetupViewFamily(FSceneViewFamily& InViewFamily) override { }
+	virtual void SetupView(FSceneViewFamily& InViewFamily, FSceneView& InView) override { }
+	virtual void BeginRenderViewFamily(FSceneViewFamily& InViewFamily) override;
+	virtual bool IsActiveThisFrame_Internal(const FSceneViewExtensionContext& Context) const override;
+	// ~End ISceneViewExtension
+};
diff --git a/Engine/Plugins/Experimental/GeometryMask/Source/GeometryMask/Private/GeometryMaskSettings.cpp b/Engine/Plugins/Experimental/GeometryMask/Source/GeometryMask/Private/GeometryMaskSettings.cpp
new file mode 100644
--- /dev/null
+++ b/Engine/Plugins/Experimental/GeometryMask/Source/GeometryMask/Private/GeometryMaskSettings.cpp
@@ -0,0 +1,13 @@
+﻿// Copyright Epic Games, Inc. All Rights Reserved.
+
+#include "GeometryMaskSettings.h"
+
+float UGeometryMaskSettings::GetDefaultResolutionMultiplier() const
+{
+	return FMath::Clamp(DefaultResolutionMultiplier, 0.25, 16.0);
+}
+
+void UGeometryMaskSettings::SetDefaultResolutionMultiplier(const float InValue)
+{
+	DefaultResolutionMultiplier = InValue;
+}
diff --git a/Engine/Plugins/Experimental/GeometryMask/Source/GeometryMask/Private/GeometryMaskSettings.h b/Engine/Plugins/Experimental/GeometryMask/Source/GeometryMask/Private/GeometryMaskSettings.h
new file mode 100644
--- /dev/null
+++ b/Engine/Plugins/Experimental/GeometryMask/Source/GeometryMask/Private/GeometryMaskSettings.h
@@ -0,0 +1,31 @@
+﻿// Copyright Epic Games, Inc. All Rights Reserved.
+
+#pragma once
+
+#include "Engine/DeveloperSettings.h"
+
+#include "GeometryMaskSettings.generated.h"
+
+UCLASS(Config = "Plugins", DefaultConfig, MinimalAPI)
+class UGeometryMaskSettings : public UDeveloperSettings
+{
+	GENERATED_BODY()
+
+public:
+	virtual FName GetContainerName() const override { return FName(TEXT("Project")); }
+	virtual FName GetCategoryName() const override { return FName(TEXT("Plugins")); }
+	virtual FName GetSectionName() const override { return FName(TEXT("Geometry Mask")); }
+
+#if WITH_EDITOR
+	virtual FText GetSectionText() const override { return NSLOCTEXT("GeometryMaskSettings", "SectionText", "Geometry Mask"); };
+	virtual FText GetSectionDescription() const override { return NSLOCTEXT("GeometryMaskSettings", "SectionDescription", "Manage Geometry Mask settings."); };
+#endif
+
+	float GetDefaultResolutionMultiplier() const;
+	void SetDefaultResolutionMultiplier(const float InValue);
+
+private:
+	/** Set to enable super-sampling, ie. 2.0 will render at twice the resolution. */
+	UPROPERTY(Config, EditAnywhere, Category = "Canvas", Getter, Setter, meta = (ClampMin = 0.25, ClampMax = 16.0, AllowPrivateAccess = "true"))
+	float DefaultResolutionMultiplier = 1.0;
+};
diff --git a/Engine/Plugins/Experimental/GeometryMask/Source/GeometryMask/Private/GeometryMaskSubsystem.cpp b/Engine/Plugins/Experimental/GeometryMask/Source/GeometryMask/Private/GeometryMaskSubsystem.cpp
new file mode 100644
--- /dev/null
+++ b/Engine/Plugins/Experimental/GeometryMask/Source/GeometryMask/Private/GeometryMaskSubsystem.cpp
@@ -0,0 +1,160 @@
+// Copyright Epic Games, Inc. All Rights Reserved.
+
+#include "GeometryMaskSubsystem.h"
+
+#include "Algo/RemoveIf.h"
+#include "Engine/Engine.h"
+#include "Engine/GameViewportClient.h"
+#include "Engine/LocalPlayer.h"
+#include "Engine/World.h"
+#include "GameFramework/PlayerController.h"
+#include "GeometryMaskCanvas.h"
+#include "GeometryMaskCanvasResource.h"
+#include "SceneView.h"
+#include "UnrealClient.h"
+
+UGeometryMaskCanvas* UGeometryMaskSubsystem::GetNamedCanvas(FName InName)
+{
+	if (InName.IsNone())
+	{
+		static FName DefaultCanvasName = TEXT("Default");
+		InName = DefaultCanvasName;
+	}
+
+	const FName ObjectName = MakeUniqueObjectName(this, UGeometryMaskCanvas::StaticClass(), FName(FString::Printf(TEXT("GeometryMaskCanvas_%s_"), *InName.ToString())));
+	if (TObjectPtr<UGeometryMaskCanvas>* FoundCanvas = NamedCanvases.Find(InName))
+	{
+		return *FoundCanvas;
+	}
+
+	const TObjectPtr<UGeometryMaskCanvas>& NewCanvas = NamedCanvases.Emplace(InName, NewObject<UGeometryMaskCanvas>(this, ObjectName));
+	NewCanvas->CanvasName = InName;
+	AssignResourceToCanvas(NewCanvas);
+
+	NewCanvas->OnActivated().BindUObject(this, &UGeometryMaskSubsystem::OnCanvasActivated, NewCanvas.Get());
+	NewCanvas->OnDeactivated().BindUObject(this, &UGeometryMaskSubsystem::OnCanvasDeactivated, NewCanvas.Get());
+	
+	OnGeometryMaskCanvasCreatedDelegate.Broadcast(NewCanvas);
+
+	return NewCanvas;
+}
+
+TArray<FName> UGeometryMaskSubsystem::GetCanvasNames()
+{
+	if (const UGeometryMaskSubsystem* Subsystem = GEngine->GetEngineSubsystem<UGeometryMaskSubsystem>())
+	{
+		TArray<FName> CanvasNames;
+		Subsystem->NamedCanvases.GenerateKeyArray(CanvasNames);
+		return CanvasNames;
+	}
+	
+	return {};
+}
+
+const TArray<TObjectPtr<UGeometryMaskCanvasResource>>& UGeometryMaskSubsystem::GetCanvasResources() const
+{
+	return CanvasResources;
+}
+
+void UGeometryMaskSubsystem::Update(
+	UWorld* InWorld,
+	FSceneViewFamily& InViewFamily)
+{
+	if (UGeometryMaskSubsystem* Subsystem = GEngine->GetEngineSubsystem<UGeometryMaskSubsystem>())
+	{
+		for (const FSceneView*& View : InViewFamily.Views)
+		{
+			FSceneView* MutableView = const_cast<FSceneView*>(View);
+			for (const TPair<FName, TObjectPtr<UGeometryMaskCanvas>>& NamedCanvas : Subsystem->NamedCanvases)
+			{
+				NamedCanvas.Value->Update(InWorld, *MutableView);			
+			}
+
+			for (const TObjectPtr<UGeometryMaskCanvasResource>& Resource : CanvasResources)
+			{
+				Resource->Update(InWorld, *MutableView);
+			}
+		}
+	}
+}
+
+int32 UGeometryMaskSubsystem::RemoveWithoutWriters()
+{
+	int32 NumRemoved = 0;
+	
+	TMap<FName, TObjectPtr<UGeometryMaskCanvas>> UsedCanvases;
+	UsedCanvases.Reserve(NamedCanvases.Num());
+	
+	for (const TPair<FName, TObjectPtr<UGeometryMaskCanvas>>& NamedCanvas : NamedCanvases)
+	{
+		if (NamedCanvas.Key == NAME_None
+			|| !NamedCanvas.Value->GetWriters().IsEmpty())
+		{
+			UsedCanvases.Emplace(NamedCanvas.Key, NamedCanvas.Value);
+		}
+		else
+		{
+			NamedCanvas.Value->FreeResource();
+			++NumRemoved;
+		}
+	}
+	
+	NamedCanvases = UsedCanvases;
+
+	return NumRemoved;
+}
+
+void UGeometryMaskSubsystem::AssignResourceToCanvas(UGeometryMaskCanvas* InCanvas)
+{
+	UGeometryMaskCanvasResource* AvailableResource = nullptr;
+	EGeometryMaskColorChannel AvailableChannel = EGeometryMaskColorChannel::None;
+	for (UGeometryMaskCanvasResource* CanvasResource : CanvasResources)
+	{
+		if (AvailableChannel = CanvasResource->GetNextAvailableColorChannel();
+			AvailableChannel != EGeometryMaskColorChannel::None)
+		{
+			AvailableResource = CanvasResource;
+		}
+	}
+	
+	// Nothing available, create new resource
+	if (AvailableChannel == EGeometryMaskColorChannel::None)
+	{
+		AvailableResource = CanvasResources.Emplace_GetRef(NewObject<UGeometryMaskCanvasResource>(this));
+		AvailableChannel = AvailableResource->GetNextAvailableColorChannel();
+	}
+
+	AvailableResource->Checkout(AvailableChannel, InCanvas->GetCanvasName());
+	InCanvas->AssignResource(AvailableResource, AvailableChannel);
+}
+
+void UGeometryMaskSubsystem::OnCanvasActivated(UGeometryMaskCanvas* InCanvas)
+{
+	if (InCanvas->IsDefaultCanvas())
+	{
+		return;
+	}
+
+	// Already has a resource
+	if (InCanvas->GetResource())
+	{
+		return;
+	}
+
+	// Provide a new resource for the canvas to write to
+	AssignResourceToCanvas(InCanvas);
+}
+
+void UGeometryMaskSubsystem::OnCanvasDeactivated(UGeometryMaskCanvas* InCanvas)
+{
+	if (InCanvas->IsDefaultCanvas())
+	{
+		return;
+	}
+
+	if (const UGeometryMaskCanvasResource* CanvasResource = InCanvas->GetResource())
+	{
+		// Resource assigned, so free it up
+		InCanvas->FreeResource();
+	}
+}
diff --git a/Engine/Plugins/Experimental/GeometryMask/Source/GeometryMask/Private/GeometryMaskTypes.cpp b/Engine/Plugins/Experimental/GeometryMask/Source/GeometryMask/Private/GeometryMaskTypes.cpp
new file mode 100644
--- /dev/null
+++ b/Engine/Plugins/Experimental/GeometryMask/Source/GeometryMask/Private/GeometryMaskTypes.cpp
@@ -0,0 +1,117 @@
+﻿// Copyright Epic Games, Inc. All Rights Reserved.
+
+#include "GeometryMaskTypes.h"
+
+#include "Engine/Engine.h"
+#include "GeometryMaskSubsystem.h"
+#include "UObject/UObjectThreadContext.h"
+
+namespace UE::GeometryMask
+{
+	EGeometryMaskColorChannel VectorToMaskChannel(const FLinearColor& InVector)
+	{
+		float CurrentMax = std::numeric_limits<float>::min();
+		int32 MaxIdx = 0;
+		for (int32 ElementIdx = 0; ElementIdx < 4; ++ElementIdx)
+		{
+			const float CurrentValue = InVector.Component(ElementIdx);			
+			if (CurrentValue > CurrentMax)
+			{
+				CurrentMax = CurrentValue;
+				MaxIdx = ElementIdx;
+			}
+		}
+		
+		return static_cast<EGeometryMaskColorChannel>(MaxIdx);
+	}
+
+	EGeometryMaskColorChannel GetValidMaskChannel(
+		EGeometryMaskColorChannel InColorChannel,
+		bool bInIncludeAlpha)
+	{
+		return FMath::Clamp(InColorChannel, EGeometryMaskColorChannel::Red, bInIncludeAlpha ? EGeometryMaskColorChannel::Alpha : EGeometryMaskColorChannel::Blue);
+	}
+
+	FStringView ChannelToString(EGeometryMaskColorChannel InColorChannel)
+	{
+		return MaskChannelEnumToString[FMath::Clamp(InColorChannel, EGeometryMaskColorChannel::Red, EGeometryMaskColorChannel::Num)];
+	}
+}
+
+UTextureRenderTarget2D* UGeometryMaskCanvasReferenceComponentBase::GetTexture()
+{
+	if (const UGeometryMaskCanvas* Canvas = CanvasWeak.Get())
+	{
+		return Canvas->GetTexture();
+	}
+	
+	return nullptr;
+}
+
+void UGeometryMaskCanvasReferenceComponentBase::BeginPlay()
+{
+	Super::BeginPlay();
+
+	TryResolveCanvas();
+}
+
+void UGeometryMaskCanvasReferenceComponentBase::PostLoad()
+{
+	Super::PostLoad();
+
+	if (!IsTemplate())
+	{
+		TryResolveCanvas();
+	}
+}
+
+void UGeometryMaskCanvasReferenceComponentBase::OnRegister()
+{
+	Super::OnRegister();
+
+	if (!IsTemplate())
+	{
+		TryResolveCanvas();
+	}
+}
+
+bool UGeometryMaskCanvasReferenceComponentBase::TryResolveNamedCanvas(FName InCanvasName)
+{
+	auto BroadcastSetCanvas = [this](const UGeometryMaskCanvas* InCanvas)
+	{
+		ReceiveSetCanvas(InCanvas);
+		OnSetCanvasDelegate.Broadcast(InCanvas);
+	};
+
+	// Get current canvas
+	UGeometryMaskCanvas* Canvas = CanvasWeak.Get();
+	
+	if (Canvas)
+	{
+		// Canvas set, no need to resolve it
+		if (Canvas->GetFName() == InCanvasName)
+		{
+			return true;
+		}
+
+		// Canvas is not the correct one
+		CanvasWeak.Reset();
+	}
+
+	if (UGeometryMaskSubsystem* Subsystem = GEngine->GetEngineSubsystem<UGeometryMaskSubsystem>())
+	{
+		Canvas = Subsystem->GetNamedCanvas(InCanvasName);
+
+		CanvasWeak = Canvas;
+		
+		if (!FUObjectThreadContext::Get().IsRoutingPostLoad)
+		{
+			if (Canvas)
+			{
+				BroadcastSetCanvas(Canvas);
+			}
+		}
+	}
+
+	return CanvasWeak.IsValid();
+}
diff --git a/Engine/Plugins/Experimental/GeometryMask/Source/GeometryMask/Private/GeometryMaskWorldSubsystem.cpp b/Engine/Plugins/Experimental/GeometryMask/Source/GeometryMask/Private/GeometryMaskWorldSubsystem.cpp
new file mode 100644
--- /dev/null
+++ b/Engine/Plugins/Experimental/GeometryMask/Source/GeometryMask/Private/GeometryMaskWorldSubsystem.cpp
@@ -0,0 +1,17 @@
+﻿// Copyright Epic Games, Inc. All Rights Reserved.
+
+#include "GeometryMaskWorldSubsystem.h"
+
+#include "Engine/Engine.h"
+#include "Engine/World.h"
+#include "GeometryMaskSVE.h"
+#include "SceneViewExtension.h"
+
+void UGeometryMaskWorldSubsystem::Initialize(FSubsystemCollectionBase& Collection)
+{
+	Super::Initialize(Collection);
+
+	UWorld* World = GetWorld();
+
+	GeometryMaskSceneViewExtension = FSceneViewExtensions::NewExtension<FGeometryMaskSceneViewExtension>(World);
+}
diff --git a/Engine/Plugins/Experimental/GeometryMask/Source/GeometryMask/Private/GeometryMaskWriteComponent.cpp b/Engine/Plugins/Experimental/GeometryMask/Source/GeometryMask/Private/GeometryMaskWriteComponent.cpp
new file mode 100644
--- /dev/null
+++ b/Engine/Plugins/Experimental/GeometryMask/Source/GeometryMask/Private/GeometryMaskWriteComponent.cpp
@@ -0,0 +1,433 @@
+﻿// Copyright Epic Games, Inc. All Rights Reserved.
+
+#include "GeometryMaskWriteComponent.h"
+
+#include "Algo/RemoveIf.h"
+#include "Async/ParallelFor.h"
+#include "BatchedElements.h"
+#include "CanvasTypes.h"
+#include "Components/DynamicMeshComponent.h"
+#include "Components/StaticMeshComponent.h"
+#include "Engine/StaticMesh.h"
+#include "GeometryMaskCanvas.h"
+#include "GlobalRenderResources.h"
+#include "StaticMeshResources.h"
+
+void UGeometryMaskWriteMeshComponent::SetParameters(FGeometryMaskWriteParameters& InParameters)
+{
+	Parameters = InParameters;
+
+	UGeometryMaskCanvas* Canvas = CanvasWeak.Get();
+
+	// We have changed canvas, unregister this writer
+	if (Canvas && Canvas->GetFName() != InParameters.CanvasName)
+	{
+		Canvas->RemoveWriter(this);
+		CanvasWeak.Reset();
+	}
+
+	// Update canvas
+	if (!CanvasWeak.IsValid())
+	{
+		TryResolveCanvas();
+	}
+}
+
+void UGeometryMaskWriteMeshComponent::DrawToCanvas(FCanvas* InCanvas)
+{
+	TRACE_CPUPROFILER_EVENT_SCOPE(UGeometryMaskWriteMeshComponent::DrawToCanvas);
+
+	if (!bWriteWhenHidden)
+	{
+		if (const AActor* Owner = GetOwner())
+		{
+			if (Owner->IsHidden())
+			{
+				return;
+			}
+		}
+	}
+
+	UpdateCachedData();
+
+	FColor Color = UE::GeometryMask::MaskChannelEnumToColor[GetParameters().ColorChannel];
+	ESimpleElementBlendMode ElementBlendMode = SE_BLEND_Additive;
+	if (GetParameters().OperationType == EGeometryMaskCompositeOperation::Subtract)
+	{
+		Color = FColor(255 - Color.R, 255 - Color.G, 255 - Color.B, 255);
+		ElementBlendMode = SE_BLEND_Modulate;
+	}
+
+	// @todo: cached mesh data keyed by SM name, but transforms also tied to SM, not SMC - will break if same SM used in multiple SMC's
+
+	// Write to Canvas BatchElements
+	{
+		FHitProxyId CanvasHitProxyId = InCanvas->GetHitProxyId();
+
+		TArray<FName> KeyNames;
+		CachedMeshData.GenerateKeyArray(KeyNames);
+
+		TRACE_CPUPROFILER_EVENT_SCOPE(UGeometryMaskWriteMeshComponent::DrawToCanvas::WriteToCanvas);
+
+		ParallelFor(KeyNames.Num(), [&](int32 TaskIdx)
+		{
+			TRACE_CPUPROFILER_EVENT_SCOPE(UGeometryMaskWriteMeshComponent::DrawToCanvas::WriteToCanvas::Task);
+
+			const FName KeyName = KeyNames[TaskIdx];
+			FTransform LocalToWorld;
+			
+			if (const USceneComponent* Component = CachedComponents[KeyName].Get())
+			{
+				LocalToWorld = Component->GetComponentToWorld();
+			}
+			else
+			{
+				// @todo: remove from cache
+				return;
+			}
+
+			const FGeometryMaskBatchElementData& MeshBatchElementData = CachedMeshData[KeyName];
+
+			{
+				FScopeLock CanvasLock(&CanvasCS);
+				InCanvas->PushAbsoluteTransform(LocalToWorld.ToMatrixWithScale());	
+			}
+			
+			{
+				FBatchedElements* CanvasTriangleBatchedElements = nullptr;				
+				{
+					FScopeLock CanvasLock(&CanvasCS);
+					CanvasTriangleBatchedElements = InCanvas->GetBatchedElements(FCanvas::ET_Triangle);
+				}
+				
+				CanvasTriangleBatchedElements->ReserveVertices(MeshBatchElementData.Vertices.Num());
+				CanvasTriangleBatchedElements->ReserveTriangles(MeshBatchElementData.NumTriangles, GWhiteTexture, ElementBlendMode);
+				
+				auto AddVertex = [CanvasTriangleBatchedElements, &MeshBatchElementData, CanvasHitProxyId, &Color](int32 VertexIdx) -> int32
+				{
+					return CanvasTriangleBatchedElements->AddVertexf(
+						MeshBatchElementData.Vertices[VertexIdx],
+						FVector2f::ZeroVector,
+						Color,
+						CanvasHitProxyId);
+				};
+
+				{
+					TRACE_CPUPROFILER_EVENT_SCOPE(UGeometryMaskWriteMeshComponent::DrawToCanvas::WriteToCanvas::AddVertices);
+					for (int32 VertexIdx = 0; VertexIdx < MeshBatchElementData.Vertices.Num(); ++VertexIdx)
+					{
+						AddVertex(VertexIdx);
+					}
+				}
+
+				{
+					TRACE_CPUPROFILER_EVENT_SCOPE(UGeometryMaskWriteMeshComponent::DrawToCanvas::WriteToCanvas::AddTriangles);
+					for (int32 VertexIdx = 0; VertexIdx <= MeshBatchElementData.Indices.Num() - 3; VertexIdx += 3)
+					{
+						const int32 V0 = MeshBatchElementData.Indices[VertexIdx];
+						const int32 V1 = MeshBatchElementData.Indices[VertexIdx + 1];
+						const int32 V2 = MeshBatchElementData.Indices[VertexIdx + 2];
+
+						CanvasTriangleBatchedElements->AddTriangle(V0, V1, V2, GWhiteTexture, ElementBlendMode);
+					}
+				}
+			}
+			
+			{
+				FScopeLock CanvasLock(&CanvasCS);
+				InCanvas->PopTransform();
+			}
+		});
+	}
+}
+
+void UGeometryMaskWriteMeshComponent::OnComponentDestroyed(bool bDestroyingHierarchy)
+{
+	Super::OnComponentDestroyed(bDestroyingHierarchy);
+
+	if (UGeometryMaskCanvas* Canvas = CanvasWeak.Get())
+	{
+		Canvas->RemoveWriter(this);
+		CanvasWeak.Reset();
+	}
+}
+
+void UGeometryMaskWriteMeshComponent::UpdateCachedData()
+{
+	if (const AActor* Owner = GetOwner())
+	{
+		TArray<UPrimitiveComponent*> PrimitiveComponents;
+		Owner->GetComponents<UPrimitiveComponent>(PrimitiveComponents);
+
+		// @todo: better check for cached data state?
+		if (LastPrimitiveComponentCount == PrimitiveComponents.Num())
+		{
+			return;
+		}
+
+		LastPrimitiveComponentCount = PrimitiveComponents.Num();
+		CachedComponents.Reserve(LastPrimitiveComponentCount);
+		CachedMeshData.Reserve(LastPrimitiveComponentCount);
+
+		UpdateCachedStaticMeshData(PrimitiveComponents);
+		UpdateCachedDynamicMeshData(PrimitiveComponents);
+	}
+}
+
+void UGeometryMaskWriteMeshComponent::UpdateCachedStaticMeshData(TConstArrayView<UPrimitiveComponent*> InPrimitiveComponents)
+{
+	TArray<UStaticMeshComponent*> StaticMeshComponents;
+	StaticMeshComponents.Reserve(InPrimitiveComponents.Num());
+	
+	Algo::TransformIf(
+		InPrimitiveComponents,
+		StaticMeshComponents,
+		[](const UPrimitiveComponent* InComponent)
+		{
+			return InComponent->IsA<UStaticMeshComponent>();
+		},
+		[](UPrimitiveComponent* InComponent)
+		{
+			return Cast<UStaticMeshComponent>(InComponent);
+		});
+
+	if (StaticMeshComponents.IsEmpty())
+	{
+		return;
+	}
+
+	// Remove built/already cached
+	StaticMeshComponents.SetNum(Algo::RemoveIf(StaticMeshComponents, [this](const UStaticMeshComponent* InStaticMeshComponent)
+	{
+		if (const UStaticMesh* StaticMesh = InStaticMeshComponent->GetStaticMesh())
+		{
+			if (const FGeometryMaskBatchElementData* CachedData = CachedMeshData.Find(StaticMesh->GetFName()))
+			{
+				if (CachedData->Vertices.Num() == StaticMesh->GetNumVertices(0))
+				{
+					return true; // Cached, remove from "to build" list							
+				}
+			}
+		}
+
+		return false;
+	}));
+
+	// Subscribe to change events
+	{
+		#if WITH_EDITOR
+		for (UStaticMeshComponent* StaticMeshComponent : StaticMeshComponents)
+		{
+			StaticMeshComponent->OnStaticMeshChanged().AddUObject(this, &UGeometryMaskWriteMeshComponent::OnStaticMeshChanged);
+		}
+		#endif
+	}
+
+	// Used for cache lookup
+	TArray<FName> StaticMeshObjectNames;
+	StaticMeshObjectNames.Reserve(StaticMeshComponents.Num());
+	
+	TArray<FStaticMeshLODResources*> StaticMeshResources;
+	StaticMeshResources.Reserve(StaticMeshComponents.Num());
+
+	// Collect valid mesh resources
+	{
+		TRACE_CPUPROFILER_EVENT_SCOPE(UGeometryMaskWriteMeshComponent::DrawToCanvas::CollectMeshResources);
+		for (UStaticMeshComponent* StaticMeshComponent : StaticMeshComponents)
+		{
+			if (UStaticMesh* StaticMesh = StaticMeshComponent->GetStaticMesh())
+			{
+				if (!StaticMesh->HasValidRenderData())
+				{
+					continue;
+				}
+			
+				if (FStaticMeshRenderData* RenderData = StaticMesh->GetRenderData())
+				{
+					if (RenderData->LODResources.IsEmpty())
+						{
+							continue;
+						}
+
+						if (RenderData->LODResources[0].GetNumVertices() == 0)
+						{
+							continue;
+						}
+
+						StaticMeshObjectNames.Add(StaticMesh->GetFName());
+						StaticMeshResources.Add(&RenderData->LODResources[0]);
+						CachedComponents.Add(StaticMesh->GetFName(), StaticMeshComponent);
+					}
+				}
+			}
+		}
+
+		// Convert mesh resources to batch elements
+		{
+			TRACE_CPUPROFILER_EVENT_SCOPE(UGeometryMaskWriteMeshComponent::DrawToCanvas::ConvertMeshResources);
+			ParallelFor(StaticMeshResources.Num(),
+			[this, &StaticMeshResources, &StaticMeshObjectNames](int32 TaskIdx)
+			{
+				TRACE_CPUPROFILER_EVENT_SCOPE(UGeometryMaskWriteMeshComponent::DrawToCanvas::ConvertMeshResources::Task);
+				
+				FStaticMeshLODResources* MeshResources = StaticMeshResources[TaskIdx];
+				const int32 NumVertices = MeshResources->GetNumVertices();
+				const int32 NumIndices = MeshResources->IndexBuffer.GetNumIndices();
+				const int32 NumTriangles = MeshResources->GetNumTriangles();
+
+				FGeometryMaskBatchElementData& MeshBatchElementData = CachedMeshData.Emplace(StaticMeshObjectNames[TaskIdx]);
+				MeshBatchElementData.Reserve(NumVertices, NumIndices, NumTriangles);
+				MeshResources->IndexBuffer.GetCopy(MeshBatchElementData.Indices);
+
+				for (int32 VertexIdx = 0; VertexIdx < NumVertices; ++VertexIdx)
+				{
+					MeshBatchElementData.Vertices.Add(FVector4f(MeshResources->VertexBuffers.PositionVertexBuffer.VertexPosition(VertexIdx)));
+				}
+			});
+		}
+}
+
+void UGeometryMaskWriteMeshComponent::UpdateCachedDynamicMeshData(TConstArrayView<UPrimitiveComponent*> InPrimitiveComponents)
+{
+	TArray<UDynamicMeshComponent*> DynamicMeshComponents;
+	DynamicMeshComponents.Reserve(InPrimitiveComponents.Num());
+	
+	Algo::TransformIf(
+		InPrimitiveComponents,
+		DynamicMeshComponents,
+		[](const UPrimitiveComponent* InComponent)
+		{
+			return InComponent->IsA<UDynamicMeshComponent>();
+		},
+		[](UPrimitiveComponent* InComponent)
+		{
+			return Cast<UDynamicMeshComponent>(InComponent);
+		});
+
+	if (DynamicMeshComponents.IsEmpty())
+	{
+		return;
+	}
+
+	// Remove built/already cached
+	DynamicMeshComponents.SetNum(Algo::RemoveIf(DynamicMeshComponents, [this](UDynamicMeshComponent* InDynamicMeshComponent)
+	{
+		if (const UDynamicMesh* DynamicMesh = InDynamicMeshComponent->GetDynamicMesh())
+		{
+			if (const FGeometryMaskBatchElementData* CachedData = CachedMeshData.Find(InDynamicMeshComponent->GetFName()))
+			{
+				if (CachedData->Vertices.Num() == DynamicMesh->GetMeshRef().VertexCount()
+					&& CachedData->ChangeStamp == DynamicMesh->GetMeshRef().GetChangeStamp())
+				{
+					return true; // Cached, remove from "to build" list							
+				}
+			}
+		}
+
+		return false;
+	}));
+
+	// Subscribe to change events
+	{
+		for (UDynamicMeshComponent* DynamicMeshComponent : DynamicMeshComponents)
+		{
+			DynamicMeshComponent->GetMesh()->SetShapeChangeStampEnabled(true);			
+			DynamicMeshComponent->OnMeshChanged.AddUObject(this, &UGeometryMaskWriteMeshComponent::OnDynamicMeshChanged, DynamicMeshComponent);
+		}
+	}
+
+	// Used for cache lookup
+	TArray<FName> DynamicMeshObjectNames;
+	DynamicMeshObjectNames.Reserve(DynamicMeshComponents.Num());
+	
+	TArray<FDynamicMesh3*> DynamicMeshes;
+	DynamicMeshes.Reserve(DynamicMeshComponents.Num());
+
+	// Collect valid meshes
+	{
+		TRACE_CPUPROFILER_EVENT_SCOPE(UGeometryMaskWriteMeshComponent::DrawToCanvas::CollectMeshResources);
+		for (UDynamicMeshComponent* DynamicMeshComponent : DynamicMeshComponents)
+		{
+			if (UDynamicMesh* DynamicMeshObject = DynamicMeshComponent->GetDynamicMesh())
+			{
+				if (FDynamicMesh3* DynamicMesh = DynamicMeshObject->GetMeshPtr())
+				{
+					DynamicMeshObjectNames.Add(DynamicMeshComponent->GetFName());
+					DynamicMeshes.Add(DynamicMesh);
+					CachedComponents.Add(DynamicMeshComponent->GetFName(), DynamicMeshComponent);
+				}
+			}
+		}
+	}
+
+	// Convert meshes to batch elements
+	{
+		TRACE_CPUPROFILER_EVENT_SCOPE(UGeometryMaskWriteMeshComponent::DrawToCanvas::ConvertMeshResources);
+		ParallelFor(DynamicMeshes.Num(),
+		[this, &DynamicMeshes, &DynamicMeshObjectNames](int32 TaskIdx)
+		{
+			TRACE_CPUPROFILER_EVENT_SCOPE(UGeometryMaskWriteMeshComponent::DrawToCanvas::ConvertMeshResources::Task);
+
+			FDynamicMesh3 CompactMesh;
+			const FDynamicMesh3* DynamicMesh = DynamicMeshes[TaskIdx];
+			CompactMesh.CompactCopy(*DynamicMesh);
+
+			const int32 NumVertices = CompactMesh.VertexCount();
+			const int32 NumIndices = CompactMesh.TriangleCount() * 3;
+			const int32 NumTriangles = CompactMesh.TriangleCount();
+
+			FGeometryMaskBatchElementData& MeshBatchElementData = CachedMeshData.Emplace(DynamicMeshObjectNames[TaskIdx]);
+			MeshBatchElementData.ChangeStamp = DynamicMesh->GetChangeStamp();
+			MeshBatchElementData.Reserve(NumVertices, NumIndices, NumTriangles);
+
+			for (const int32 VertexIdx : CompactMesh.VertexIndicesItr())
+			{
+				MeshBatchElementData.Vertices.Add(FVector4f(FVector3f(CompactMesh.GetVertex(VertexIdx))));				
+			}
+
+			for (const UE::Geometry::FIndex3i Triangle : CompactMesh.TrianglesItr())
+			{
+				MeshBatchElementData.Indices.Append({
+					static_cast<uint32>(Triangle.A),
+					static_cast<uint32>(Triangle.B),
+					static_cast<uint32>(Triangle.C)});
+			}
+		});
+	}
+}
+
+#if WITH_EDITOR
+void UGeometryMaskWriteMeshComponent::OnStaticMeshChanged(UStaticMeshComponent* InStaticMeshComponent)
+{
+	// Triggers a cache refresh
+	LastPrimitiveComponentCount = -1;
+}
+#endif
+
+void UGeometryMaskWriteMeshComponent::OnDynamicMeshChanged(UDynamicMeshComponent* InDynamicMeshComponent)
+{
+	// Triggers a specific cache refresh
+	if (FGeometryMaskBatchElementData* CachedData = CachedMeshData.Find(InDynamicMeshComponent->GetFName()))
+	{
+		CachedData->ChangeStamp = INDEX_NONE;
+	}
+	
+	// Triggers a general cache refresh
+	LastPrimitiveComponentCount = -1;
+}
+
+bool UGeometryMaskWriteMeshComponent::TryResolveCanvas()
+{
+	if (TryResolveNamedCanvas(Parameters.CanvasName))
+	{
+		if (UGeometryMaskCanvas* Canvas = CanvasWeak.Get())
+		{
+			// Register this writer
+			Parameters.ColorChannel = Canvas->GetColorChannel();
+			Canvas->AddWriter(this);
+			return true;
+		}
+	}
+
+	return false;
+}
diff --git a/Engine/Plugins/Experimental/GeometryMask/Source/GeometryMask/Private/Shaders/GeometryMaskPostProcess.h b/Engine/Plugins/Experimental/GeometryMask/Source/GeometryMask/Private/Shaders/GeometryMaskPostProcess.h
new file mode 100644
--- /dev/null
+++ b/Engine/Plugins/Experimental/GeometryMask/Source/GeometryMask/Private/Shaders/GeometryMaskPostProcess.h
@@ -0,0 +1,41 @@
+﻿// Copyright Epic Games, Inc. All Rights Reserved.
+
+#pragma once
+
+#include "RHICommandList.h"
+
+class FRenderTarget;
+
+class IGeometryMaskPostProcess
+{
+public:
+	virtual ~IGeometryMaskPostProcess() = default;
+	
+	virtual void Execute(FRenderTarget* InTexture) = 0;
+
+protected:
+	virtual void Execute_RenderThread(FRHICommandListImmediate& InRHICmdList, FRenderTarget* InTexture) = 0;
+};
+
+template <typename InParametersType>
+class TGeometryMaskPostProcess
+	: public IGeometryMaskPostProcess
+{
+public:
+	explicit TGeometryMaskPostProcess(const InParametersType& InParameters)
+		: Parameters(InParameters) { }
+
+	const InParametersType& GetParameters() const
+	{
+		return Parameters;
+	}
+	
+	void SetParameters(const InParametersType& InParameters)
+	{
+		Parameters = InParameters;
+	}
+
+protected:
+	InParametersType Parameters;
+};
+
diff --git a/Engine/Plugins/Experimental/GeometryMask/Source/GeometryMask/Private/Shaders/GeometryMaskPostProcess_Blur.cpp b/Engine/Plugins/Experimental/GeometryMask/Source/GeometryMask/Private/Shaders/GeometryMaskPostProcess_Blur.cpp
new file mode 100644
--- /dev/null
+++ b/Engine/Plugins/Experimental/GeometryMask/Source/GeometryMask/Private/Shaders/GeometryMaskPostProcess_Blur.cpp
@@ -0,0 +1,290 @@
+// Copyright Epic Games, Inc. All Rights Reserved.
+
+#include "GeometryMaskPostProcess_Blur.h"
+
+#include "GlobalShader.h"
+#include "RenderGraphBuilder.h"
+#include "RenderGraphUtils.h"
+#include "RHIStaticStates.h"
+#include "UnrealClient.h"
+
+namespace UE::GeometryMask::Private
+{
+	static constexpr int32 MaxNumSteps = 127;	
+}
+
+class FGeometryMaskBlurCS : public FGlobalShader
+{
+	DECLARE_GLOBAL_SHADER(FGeometryMaskBlurCS);
+	SHADER_USE_PARAMETER_STRUCT(FGeometryMaskBlurCS, FGlobalShader);
+
+	BEGIN_SHADER_PARAMETER_STRUCT(FParameters, )
+		SHADER_PARAMETER_SCALAR_ARRAY(int32, SampleCount, [4])
+		SHADER_PARAMETER_ARRAY(FVector4f, WeightsAndOffsets, [4 * UE::GeometryMask::Private::MaxNumSteps])
+		SHADER_PARAMETER_ARRAY(FVector4f, BufferSizeAndDirection, [4])
+		SHADER_PARAMETER(FIntPoint, InputDimensions)
+		SHADER_PARAMETER_RDG_TEXTURE_SRV(Texture2D, InputTexture)
+		SHADER_PARAMETER_SAMPLER(SamplerState, InputSampler)
+		SHADER_PARAMETER_RDG_TEXTURE_UAV(RWTexture3D<float4>, OutputTexture)
+	END_SHADER_PARAMETER_STRUCT()
+
+public:
+	static void ModifyCompilationEnvironment(const FGlobalShaderPermutationParameters& InParameters, FShaderCompilerEnvironment& OutEnvironment)
+	{
+		FGlobalShader::ModifyCompilationEnvironment(InParameters, OutEnvironment);
+
+		OutEnvironment.SetDefine(TEXT("TILE_SIZE"), FComputeShaderUtils::kGolden2DGroupSize);
+	}
+};
+
+IMPLEMENT_GLOBAL_SHADER(FGeometryMaskBlurCS, "/Plugin/GeometryMask/Private/GeometryMaskBlurCS.usf", "MainCS", SF_Compute);
+
+namespace UE::GeometryMask
+{
+	namespace Private
+	{
+		static float GetWeight(float InDistance, float InStrength)
+		{
+			if (FMath::IsNearlyZero(InStrength))
+			{
+				return 0.0f;
+			}
+
+			// from https://en.wikipedia.org/wiki/Gaussian_blur
+			const float Strength2 = InStrength * InStrength;
+			return (1.0f / FMath::Sqrt(2 * PI * Strength2)) * FMath::Exp(-(InDistance * InDistance) / (2 * Strength2));
+		}
+
+		static FVector2f GetWeightAndOffset(float InDistance, float InSigma)
+		{
+			const float Offset1 = InDistance;
+			const float Weight1 = GetWeight(Offset1, InSigma);
+
+			const float Offset2 = InDistance + 1;
+			const float Weight2 = GetWeight(Offset2, InSigma);
+
+			const float TotalWeight = Weight1 + Weight2;
+
+			float Offset = 0;
+			if (TotalWeight > 0)
+			{
+				Offset = (Weight1*Offset1 + Weight2*Offset2) / TotalWeight;
+			}
+
+			return FVector2f(TotalWeight, Offset);
+		}
+
+		static int32 ComputeWeights(int32 InKernelSize, float InSigma, TArray<FVector4f>& OutWeightsAndOffsets)
+		{
+			const int32 NumSamples = FMath::DivideAndRoundUp(InKernelSize, 2);
+
+			// We need half of the sample count array because we're packing two samples into one float4
+			OutWeightsAndOffsets.AddUninitialized(NumSamples % 2 == 0 ? NumSamples / 2 : NumSamples / 2 + 1);
+
+			OutWeightsAndOffsets[0] = FVector4f(FVector2f(GetWeight(0,InSigma), 0), GetWeightAndOffset(1, InSigma));
+			int32 SampleIndex = 1;
+			for (int32 X = 3; X < InKernelSize; X += 4)
+			{
+				OutWeightsAndOffsets[SampleIndex] = FVector4f(GetWeightAndOffset((float)X, InSigma), GetWeightAndOffset((float)(X + 2), InSigma));
+				++SampleIndex;
+			}
+
+			return NumSamples;
+		}
+
+		int32 ComputeBlurWeights(int32 InKernelSize, float InStdDev, TArray<FVector4f>& OutWeightsAndOffsets)
+		{
+			return ComputeWeights(InKernelSize, InStdDev, OutWeightsAndOffsets);
+		}
+	}
+
+	namespace Internal
+	{
+		// @see: SBackgroundBlur::ComputeEffectiveKernelSize
+		int32 ComputeEffectiveKernelSize(double InStrength)
+		{
+			// If the radius isn't set, auto-compute it based on the strength
+			// @note: * 4.0 differs from * 0.3 in the slate code as it seems to produce better auto-computed results
+			int32 KernelSize = FMath::RoundToInt(InStrength * 4.0f);
+		
+			if (KernelSize % 2 == 0)
+			{
+				++KernelSize;
+			}
+
+			static constexpr int32 MaxKernelSize = 255;
+			KernelSize = FMath::Clamp(KernelSize, 3, MaxKernelSize);
+
+			return KernelSize;
+		}
+	}
+}
+
+FGeometryMaskPostProcessParameters_Blur::FGeometryMaskPostProcessParameters_Blur()
+	: PerChannelBlurStrength({ 16, 16, 16, 16})
+{
+	bPerChannelApplyBlur.Init(false, 4);
+}
+
+FGeometryMaskPostProcess_Blur::FGeometryMaskPostProcess_Blur(const FGeometryMaskPostProcessParameters_Blur& InParameters)
+	: Super(InParameters)
+{
+}
+
+void FGeometryMaskPostProcess_Blur::Execute(FRenderTarget* InTexture)
+{
+	if (!InTexture)
+	{
+		return;
+	}
+
+	// If blur NOT applied to any channel, early-out 
+	if (!Parameters.bPerChannelApplyBlur.Contains(true))
+	{
+		return;
+	}
+	
+	ENQUEUE_RENDER_COMMAND(CanvasRenderTargetMakeCurrentCommand)(
+	[Self = AsShared(), InTexture](FRHICommandListImmediate& InRHICmdList)
+	{
+		Self->Execute_RenderThread(InRHICmdList, InTexture);
+	});
+}
+
+void FGeometryMaskPostProcess_Blur::Execute_RenderThread(
+	FRHICommandListImmediate& InRHICmdList
+	, FRenderTarget* InTexture)
+{
+	ensure(IsInRenderingThread());
+
+	static constexpr int32 ChannelCount = 4;
+	int32 PerChannelSampleCount[ChannelCount];
+	TArray<FVector4f> PerChannelWeightsAndOffsets[ChannelCount];
+
+	for (int32 ChannelIdx = 0; ChannelIdx < ChannelCount; ++ChannelIdx)
+	{
+		double BlurStrength = Parameters.PerChannelBlurStrength[ChannelIdx];
+
+		int32 BlurKernelSize = UE::GeometryMask::Internal::ComputeEffectiveKernelSize(BlurStrength);
+
+		BlurStrength = FMath::Max(0.5, BlurStrength);
+
+		// If we're not applying blur, store as sample with value 0
+		PerChannelSampleCount[ChannelIdx] =
+			Parameters.bPerChannelApplyBlur[ChannelIdx]
+				? UE::GeometryMask::Private::ComputeBlurWeights(BlurKernelSize, BlurStrength, PerChannelWeightsAndOffsets[ChannelIdx])
+				: 0;
+	}
+
+	FRDGBuilder GraphBuilder(InRHICmdList);
+	{
+		FIntPoint InputSize = InTexture->GetSizeXY();
+
+		FRDGTextureRef InputTexture = InTexture->GetRenderTargetTexture(GraphBuilder);
+		FRDGTextureSRVRef InputTexture_SRV = GraphBuilder.CreateSRV(FRDGTextureSRVDesc(InputTexture));
+		
+		FRDGTextureDesc IntermediateTextureDesc(
+			FRDGTextureDesc::Create2D(
+				InputTexture->Desc.Extent,
+				InputTexture->Desc.Format,
+				FClearValueBinding::Green,
+				TexCreate_ShaderResource | TexCreate_UAV));
+
+		FRDGTextureRef IntermediateTexture = GraphBuilder.CreateTexture(IntermediateTextureDesc, TEXT("GeometryMaskBlur.IntermediateTexture"));
+		FRDGTextureSRVRef IntermediateTexture_SRV = GraphBuilder.CreateSRV(FRDGTextureSRVDesc(IntermediateTexture));
+		FRDGTextureUAVRef IntermediateTexture_UAV = GraphBuilder.CreateUAV(FRDGTextureUAVDesc(IntermediateTexture));
+
+		FRDGTextureDesc OutputTextureDesc(
+			FRDGTextureDesc::Create2D(
+				InputTexture->Desc.Extent,
+				InputTexture->Desc.Format,
+				FClearValueBinding::Green,
+				TexCreate_ShaderResource | TexCreate_UAV));
+
+		FRDGTextureRef OutputTexture = GraphBuilder.CreateTexture(OutputTextureDesc, TEXT("GeometryMaskBlur.OutputTexture"));
+		FRDGTextureUAVRef OutputTexture_UAV = GraphBuilder.CreateUAV(FRDGTextureUAVDesc(OutputTexture));
+
+		TShaderMapRef<FGeometryMaskBlurCS> ComputeShader(GetGlobalShaderMap(GMaxRHIFeatureLevel));
+		const FIntVector NumGroups = FComputeShaderUtils::GetGroupCount(InputSize.Size(), FIntPoint(8, 8));
+
+		// Horizontal Blur
+		{
+			RDG_EVENT_SCOPE(GraphBuilder, "GeometryMaskBlurH");
+			
+			FGeometryMaskBlurCS::FParameters* PassParameters = GraphBuilder.AllocParameters<FGeometryMaskBlurCS::FParameters>();
+			{
+				for (int32 ChannelIdx = 0; ChannelIdx < ChannelCount; ++ChannelIdx)
+				{
+					int32 WeightsAndOffsetsIdxOffset = ChannelIdx * UE::GeometryMask::Private::MaxNumSteps;
+					for (int32 WeightAndOffsetIdx = 0; WeightAndOffsetIdx < PerChannelWeightsAndOffsets[ChannelIdx].Num(); ++WeightAndOffsetIdx)
+					{
+						PassParameters->WeightsAndOffsets[WeightsAndOffsetsIdxOffset + WeightAndOffsetIdx] = PerChannelWeightsAndOffsets[ChannelIdx][WeightAndOffsetIdx];
+					}
+
+					PassParameters->BufferSizeAndDirection[ChannelIdx] = FVector4f(1.0f / InputSize.X, 1.0f / InputSize.Y, 1.0f, 0.0f);
+
+					GET_SCALAR_ARRAY_ELEMENT(PassParameters->SampleCount, ChannelIdx) = PerChannelSampleCount[ChannelIdx];
+				}
+
+				PassParameters->InputDimensions = InputSize;
+				PassParameters->InputTexture = InputTexture_SRV;
+				PassParameters->InputSampler = TStaticSamplerState<SF_Bilinear>::GetRHI();
+				PassParameters->OutputTexture = IntermediateTexture_UAV;
+			}
+
+			ClearUnusedGraphResources(ComputeShader, PassParameters);
+			
+			FComputeShaderUtils::AddPass(
+				GraphBuilder,
+				RDG_EVENT_NAME("BlurH"),
+				ERDGPassFlags::Compute,
+				ComputeShader,
+				PassParameters,
+				NumGroups);
+		}
+
+		// Vertical Blur
+		{
+			RDG_EVENT_SCOPE(GraphBuilder, "GeometryMaskBlurV");
+			
+			FGeometryMaskBlurCS::FParameters* PassParameters = GraphBuilder.AllocParameters<FGeometryMaskBlurCS::FParameters>();
+			{
+				for (int32 ChannelIdx = 0; ChannelIdx < ChannelCount; ++ChannelIdx)
+				{
+					int32 WeightsAndOffsetsIdxOffset = ChannelIdx * UE::GeometryMask::Private::MaxNumSteps;
+					for (int32 WeightAndOffsetIdx = 0; WeightAndOffsetIdx < PerChannelWeightsAndOffsets[ChannelIdx].Num(); ++WeightAndOffsetIdx)
+					{
+						PassParameters->WeightsAndOffsets[WeightsAndOffsetsIdxOffset + WeightAndOffsetIdx] = PerChannelWeightsAndOffsets[ChannelIdx][WeightAndOffsetIdx];
+					}
+
+					PassParameters->BufferSizeAndDirection[ChannelIdx] = FVector4f(1.0f / InputSize.X, 1.0f / InputSize.Y, 0.0f, 1.0f);
+					GET_SCALAR_ARRAY_ELEMENT(PassParameters->SampleCount, ChannelIdx) = PerChannelSampleCount[ChannelIdx];
+				}
+
+				PassParameters->InputDimensions = InputSize;
+				PassParameters->InputTexture = IntermediateTexture_SRV;
+				PassParameters->InputSampler = TStaticSamplerState<SF_Bilinear>::GetRHI();
+				PassParameters->OutputTexture = OutputTexture_UAV;
+			}
+
+			ClearUnusedGraphResources(ComputeShader, PassParameters);
+
+			FComputeShaderUtils::AddPass(
+				GraphBuilder,
+				RDG_EVENT_NAME("BlurV"),
+				ERDGPassFlags::Compute,
+				ComputeShader,
+				PassParameters,
+				NumGroups);
+		}
+
+		// Copy back to RenderTarget
+		{
+			RDG_EVENT_SCOPE(GraphBuilder, "GeometryMaskBlurPostCopy");
+
+			AddCopyTexturePass(GraphBuilder, OutputTexture, InputTexture, FRHICopyTextureInfo());
+		}
+	}
+
+	GraphBuilder.Execute();
+}
diff --git a/Engine/Plugins/Experimental/GeometryMask/Source/GeometryMask/Private/Shaders/GeometryMaskPostProcess_Blur.h b/Engine/Plugins/Experimental/GeometryMask/Source/GeometryMask/Private/Shaders/GeometryMaskPostProcess_Blur.h
new file mode 100644
--- /dev/null
+++ b/Engine/Plugins/Experimental/GeometryMask/Source/GeometryMask/Private/Shaders/GeometryMaskPostProcess_Blur.h
@@ -0,0 +1,37 @@
+// Copyright Epic Games, Inc. All Rights Reserved.
+
+#pragma once
+
+#include "Containers/Array.h"
+#include "Containers/BitArray.h"
+#include "Containers/ContainerAllocationPolicies.h"
+#include "GeometryMaskPostProcess.h"
+#include "Templates/SharedPointer.h"
+
+namespace UE::GeometryMask::Internal
+{
+	int32 ComputeEffectiveKernelSize(double InStrength);
+}
+
+struct FGeometryMaskPostProcessParameters_Blur
+{
+	FGeometryMaskPostProcessParameters_Blur();
+	
+	TBitArray<TInlineAllocator<4>> bPerChannelApplyBlur;
+	TArray<double, TInlineAllocator<4>> PerChannelBlurStrength;
+};
+
+class FGeometryMaskPostProcess_Blur
+	: public TGeometryMaskPostProcess<FGeometryMaskPostProcessParameters_Blur>
+	, public TSharedFromThis<FGeometryMaskPostProcess_Blur>
+{
+	using Super = TGeometryMaskPostProcess<FGeometryMaskPostProcessParameters_Blur>;
+	
+public:
+	explicit FGeometryMaskPostProcess_Blur(const FGeometryMaskPostProcessParameters_Blur& InParameters);
+
+	virtual void Execute(FRenderTarget* InTexture) override;
+
+protected:
+	virtual void Execute_RenderThread(FRHICommandListImmediate& InRHICmdList, FRenderTarget* InTexture) override;
+};
diff --git a/Engine/Plugins/Experimental/GeometryMask/Source/GeometryMask/Private/Shaders/GeometryMaskPostProcess_DistanceField.cpp b/Engine/Plugins/Experimental/GeometryMask/Source/GeometryMask/Private/Shaders/GeometryMaskPostProcess_DistanceField.cpp
new file mode 100644
--- /dev/null
+++ b/Engine/Plugins/Experimental/GeometryMask/Source/GeometryMask/Private/Shaders/GeometryMaskPostProcess_DistanceField.cpp
@@ -0,0 +1,511 @@
+﻿// Copyright Epic Games, Inc. All Rights Reserved.
+
+#include "GeometryMaskPostProcess_DistanceField.h"
+
+#include "GlobalShader.h"
+#include "RenderGraphBuilder.h"
+#include "RenderGraphUtils.h"
+#include "RHIStaticStates.h"
+#include "UnrealClient.h"
+
+namespace UE::GeometryMask::Private
+{
+	static constexpr int32 MaxNumChannels = 4;
+	static constexpr int32 NumNeighbors = 8;
+
+	static FIntVector4 CalculateStepCount(const FIntVector4& InRadiusSizes)
+	{
+		return FIntVector4(
+			InRadiusSizes[0] == 0 ? 0 : FMath::CeilToInt32(FMath::Log2(static_cast<double>(InRadiusSizes[0]))) - 1,
+			InRadiusSizes[1] == 0 ? 0 : FMath::CeilToInt32(FMath::Log2(static_cast<double>(InRadiusSizes[1]))) - 1,
+			InRadiusSizes[2] == 0 ? 0 : FMath::CeilToInt32(FMath::Log2(static_cast<double>(InRadiusSizes[2]))) - 1,
+			InRadiusSizes[3] == 0 ? 0 : FMath::CeilToInt32(FMath::Log2(static_cast<double>(InRadiusSizes[3]))) - 1);
+	}
+
+	static FIntVector4 CalculateStepSize(int32 InStepIdx, const FIntVector4& InMaxSteps, FIntVector4& OutMask)
+	{
+		FIntVector4 StepDelta = InMaxSteps - FIntVector4(InStepIdx, InStepIdx, InStepIdx, InStepIdx);
+		StepDelta.X = FMath::Min(InMaxSteps.X, StepDelta.X);
+		StepDelta.Y = FMath::Min(InMaxSteps.Y, StepDelta.Y);
+		StepDelta.Z = FMath::Min(InMaxSteps.Z, StepDelta.Z);
+		StepDelta.W = FMath::Min(InMaxSteps.W, StepDelta.W);
+
+		FIntVector4 StepSize(
+			FMath::CeilToInt32(FMath::Pow<double>(2, StepDelta.X)),
+			FMath::CeilToInt32(FMath::Pow<double>(2, StepDelta.Y)),
+			FMath::CeilToInt32(FMath::Pow<double>(2, StepDelta.Z)),
+			FMath::CeilToInt32(FMath::Pow<double>(2, StepDelta.W)));
+
+		OutMask[0] = StepSize.X > 0;
+		OutMask[1] = StepSize.Y > 0;
+		OutMask[2] = StepSize.Z > 0;
+		OutMask[3] = StepSize.W > 0;
+
+		return StepSize;
+	}
+}
+
+class FGeometryMaskJFInitCSBase : public FGlobalShader
+{
+	DECLARE_INLINE_TYPE_LAYOUT(FGeometryMaskJFInitCSBase, NonVirtual);
+	SHADER_USE_PARAMETER_STRUCT(FGeometryMaskJFInitCSBase, FGlobalShader);
+
+	BEGIN_SHADER_PARAMETER_STRUCT(FParameters, )
+		SHADER_PARAMETER_ARRAY(FVector4f, Offsets, [UE::GeometryMask::Private::MaxNumChannels * (UE::GeometryMask::Private::NumNeighbors + 1)])
+		SHADER_PARAMETER(FIntPoint, InputDimensions)
+		SHADER_PARAMETER(FVector2f, OneOverInputDimensions)
+		SHADER_PARAMETER(FVector2f, UVRatioAdjustment)
+		SHADER_PARAMETER(FIntVector4, ChannelMap)
+		SHADER_PARAMETER_RDG_TEXTURE_SRV(Texture2D, InputTexture)
+		SHADER_PARAMETER_SAMPLER(SamplerState, InputSampler)
+		SHADER_PARAMETER_RDG_BUFFER_UAV(RWBuffer<float4>, OutputBuffer)
+	END_SHADER_PARAMETER_STRUCT()
+
+public:
+	static void ModifyCompilationEnvironment(const FGlobalShaderPermutationParameters& InParameters, FShaderCompilerEnvironment& OutEnvironment)
+	{
+		FGlobalShader::ModifyCompilationEnvironment(InParameters, OutEnvironment);
+
+		OutEnvironment.SetDefine(TEXT("TILE_SIZE"), FComputeShaderUtils::kGolden2DGroupSize);
+	}
+};
+
+template <int32 NumChannels>
+class TGeometryMaskJFInitCS : public FGeometryMaskJFInitCSBase
+{
+	DECLARE_GLOBAL_SHADER(TGeometryMaskJFInitCS);
+	SHADER_USE_PARAMETER_STRUCT(TGeometryMaskJFInitCS, FGeometryMaskJFInitCSBase);
+
+public:
+	static void ModifyCompilationEnvironment(const FGlobalShaderPermutationParameters& InParameters, FShaderCompilerEnvironment& OutEnvironment)
+	{
+		FGlobalShader::ModifyCompilationEnvironment(InParameters, OutEnvironment);
+
+		OutEnvironment.SetDefine(TEXT("NUM_CHANNELS"), NumChannels);
+	}
+};
+
+using FGeometryMaskJFInitCS1 = TGeometryMaskJFInitCS<1>;
+using FGeometryMaskJFInitCS2 = TGeometryMaskJFInitCS<2>;
+using FGeometryMaskJFInitCS3 = TGeometryMaskJFInitCS<3>;
+using FGeometryMaskJFInitCS4 = TGeometryMaskJFInitCS<4>;
+
+IMPLEMENT_SHADER_TYPE(template<>, FGeometryMaskJFInitCS1, TEXT("/Plugin/GeometryMask/Private/GeometryMaskJFInitCS.usf"), TEXT("MainCS"), SF_Compute);
+IMPLEMENT_SHADER_TYPE(template<>, FGeometryMaskJFInitCS2, TEXT("/Plugin/GeometryMask/Private/GeometryMaskJFInitCS.usf"), TEXT("MainCS"), SF_Compute);
+IMPLEMENT_SHADER_TYPE(template<>, FGeometryMaskJFInitCS3, TEXT("/Plugin/GeometryMask/Private/GeometryMaskJFInitCS.usf"), TEXT("MainCS"), SF_Compute);
+IMPLEMENT_SHADER_TYPE(template<>, FGeometryMaskJFInitCS4, TEXT("/Plugin/GeometryMask/Private/GeometryMaskJFInitCS.usf"), TEXT("MainCS"), SF_Compute);
+
+class FGeometryMaskJFStepCSBase : public FGlobalShader
+{
+	DECLARE_INLINE_TYPE_LAYOUT(FGeometryMaskJFStepCSBase, NonVirtual);
+	SHADER_USE_PARAMETER_STRUCT(FGeometryMaskJFStepCSBase, FGlobalShader);
+
+	BEGIN_SHADER_PARAMETER_STRUCT(FParameters, )
+		SHADER_PARAMETER(FIntVector4, StepSize)
+		SHADER_PARAMETER(FIntVector4, ChannelMask)
+		SHADER_PARAMETER(FIntPoint, InputDimensions)
+		SHADER_PARAMETER(FVector2f, OneOverInputDimensions)
+		SHADER_PARAMETER(FVector2f, UVRatioAdjustment)
+		SHADER_PARAMETER_RDG_BUFFER_SRV(Buffer<float4>, InputBuffer)
+		SHADER_PARAMETER_RDG_BUFFER_UAV(RWBuffer<float4>, OutputBuffer)
+	END_SHADER_PARAMETER_STRUCT()
+
+public:
+	static void ModifyCompilationEnvironment(const FGlobalShaderPermutationParameters& InParameters, FShaderCompilerEnvironment& OutEnvironment)
+	{
+		FGlobalShader::ModifyCompilationEnvironment(InParameters, OutEnvironment);
+
+		OutEnvironment.SetDefine(TEXT("TILE_SIZE"), FComputeShaderUtils::kGolden2DGroupSize);
+	}
+};
+
+template <int32 NumChannels>
+class TGeometryMaskJFStepCS : public FGeometryMaskJFStepCSBase
+{
+	DECLARE_GLOBAL_SHADER(TGeometryMaskJFStepCS);
+	SHADER_USE_PARAMETER_STRUCT(TGeometryMaskJFStepCS, FGeometryMaskJFStepCSBase);
+
+public:
+	static void ModifyCompilationEnvironment(const FGlobalShaderPermutationParameters& InParameters, FShaderCompilerEnvironment& OutEnvironment)
+	{
+		FGeometryMaskJFStepCSBase::ModifyCompilationEnvironment(InParameters, OutEnvironment);
+
+		OutEnvironment.SetDefine(TEXT("NUM_CHANNELS"), NumChannels);
+	}
+};
+
+using FGeometryMaskJFStepCS1 = TGeometryMaskJFStepCS<1>;
+using FGeometryMaskJFStepCS2 = TGeometryMaskJFStepCS<2>;
+using FGeometryMaskJFStepCS3 = TGeometryMaskJFStepCS<3>;
+using FGeometryMaskJFStepCS4 = TGeometryMaskJFStepCS<4>;
+
+IMPLEMENT_SHADER_TYPE(template<>, FGeometryMaskJFStepCS1, TEXT("/Plugin/GeometryMask/Private/GeometryMaskJFStepCS.usf"), TEXT("MainCS"), SF_Compute);
+IMPLEMENT_SHADER_TYPE(template<>, FGeometryMaskJFStepCS2, TEXT("/Plugin/GeometryMask/Private/GeometryMaskJFStepCS.usf"), TEXT("MainCS"), SF_Compute);
+IMPLEMENT_SHADER_TYPE(template<>, FGeometryMaskJFStepCS3, TEXT("/Plugin/GeometryMask/Private/GeometryMaskJFStepCS.usf"), TEXT("MainCS"), SF_Compute);
+IMPLEMENT_SHADER_TYPE(template<>, FGeometryMaskJFStepCS4, TEXT("/Plugin/GeometryMask/Private/GeometryMaskJFStepCS.usf"), TEXT("MainCS"), SF_Compute);
+
+class FGeometryMaskJFtoDFCSBase : public FGlobalShader
+{
+	DECLARE_INLINE_TYPE_LAYOUT(FGeometryMaskJFtoDFCSBase, NonVirtual);
+	SHADER_USE_PARAMETER_STRUCT(FGeometryMaskJFtoDFCSBase, FGlobalShader);
+
+	BEGIN_SHADER_PARAMETER_STRUCT(FParameters, )
+		SHADER_PARAMETER(FIntPoint, InputDimensions)
+		SHADER_PARAMETER(FVector2f, OneOverInputDimensions)
+		SHADER_PARAMETER(FVector2f, UVRatioAdjustment)
+		SHADER_PARAMETER(FIntVector4, ChannelMap)
+		SHADER_PARAMETER(FVector4f, StepDistanceMultipliers)
+		SHADER_PARAMETER_RDG_TEXTURE_SRV(Texture2D, OriginalInputTexture)
+		SHADER_PARAMETER_SAMPLER(SamplerState, OriginalInputSampler)
+		SHADER_PARAMETER_RDG_BUFFER_SRV(Buffer<float4>, InputBuffer)
+		SHADER_PARAMETER_RDG_TEXTURE_UAV(RWTexture2D<float4>, OutputTexture)
+	END_SHADER_PARAMETER_STRUCT()
+
+public:
+	static void ModifyCompilationEnvironment(const FGlobalShaderPermutationParameters& InParameters, FShaderCompilerEnvironment& OutEnvironment)
+	{
+		FGlobalShader::ModifyCompilationEnvironment(InParameters, OutEnvironment);
+
+		OutEnvironment.SetDefine(TEXT("TILE_SIZE"), FComputeShaderUtils::kGolden2DGroupSize);
+	}
+};
+
+template <int32 NumChannels>
+class TGeometryMaskJFtoDFCS : public FGeometryMaskJFtoDFCSBase
+{
+	DECLARE_GLOBAL_SHADER(TGeometryMaskJFtoDFCS);
+	SHADER_USE_PARAMETER_STRUCT(TGeometryMaskJFtoDFCS, FGeometryMaskJFtoDFCSBase);
+
+public:
+	static void ModifyCompilationEnvironment(const FGlobalShaderPermutationParameters& InParameters, FShaderCompilerEnvironment& OutEnvironment)
+	{
+		FGeometryMaskJFtoDFCSBase::ModifyCompilationEnvironment(InParameters, OutEnvironment);
+
+		OutEnvironment.SetDefine(TEXT("NUM_CHANNELS"), NumChannels);
+	}
+};
+
+using FGeometryMaskJFtoDFCS1 = TGeometryMaskJFtoDFCS<1>;
+using FGeometryMaskJFtoDFCS2 = TGeometryMaskJFtoDFCS<2>;
+using FGeometryMaskJFtoDFCS3 = TGeometryMaskJFtoDFCS<3>;
+using FGeometryMaskJFtoDFCS4 = TGeometryMaskJFtoDFCS<4>;
+
+IMPLEMENT_SHADER_TYPE(template<>, FGeometryMaskJFtoDFCS1, TEXT("/Plugin/GeometryMask/Private/GeometryMaskJFtoDFCS.usf"), TEXT("MainCS"), SF_Compute);
+IMPLEMENT_SHADER_TYPE(template<>, FGeometryMaskJFtoDFCS2, TEXT("/Plugin/GeometryMask/Private/GeometryMaskJFtoDFCS.usf"), TEXT("MainCS"), SF_Compute);
+IMPLEMENT_SHADER_TYPE(template<>, FGeometryMaskJFtoDFCS3, TEXT("/Plugin/GeometryMask/Private/GeometryMaskJFtoDFCS.usf"), TEXT("MainCS"), SF_Compute);
+IMPLEMENT_SHADER_TYPE(template<>, FGeometryMaskJFtoDFCS4, TEXT("/Plugin/GeometryMask/Private/GeometryMaskJFtoDFCS.usf"), TEXT("MainCS"), SF_Compute);
+
+FGeometryMaskPostProcessParameters_DistanceField::FGeometryMaskPostProcessParameters_DistanceField()
+{
+	bPerChannelCalculateDF.Init(false, 4);
+	PerChannelRadius.Init(0.0, 4);
+}
+
+FGeometryMaskPostProcess_DistanceField::FGeometryMaskPostProcess_DistanceField(const FGeometryMaskPostProcessParameters_DistanceField& InParameters)
+	: Super(InParameters)
+{
+}
+
+void FGeometryMaskPostProcess_DistanceField::Execute(FRenderTarget* InTexture)
+{
+	if (!InTexture)
+	{
+		return;
+	}
+
+	// If blur NOT applied to any channel, early-out 
+	if (!Parameters.bPerChannelCalculateDF.Contains(true))
+	{
+		return;
+	}
+	
+	ENQUEUE_RENDER_COMMAND(CanvasRenderTargetMakeCurrentCommand)(
+	[Self = AsShared(), InTexture](FRHICommandListImmediate& InRHICmdList)
+	{
+		Self->Execute_RenderThread(InRHICmdList, InTexture);
+	});
+}
+
+void FGeometryMaskPostProcess_DistanceField::Execute_RenderThread(
+	FRHICommandListImmediate& InRHICmdList,
+	FRenderTarget* InTexture)
+{
+	ensure(IsInRenderingThread());
+
+	FRDGBuilder GraphBuilder(InRHICmdList);
+	{
+		FIntPoint InputSize = InTexture->GetSizeXY();
+		FVector2f OneOverInputSize = FVector2f::One() / InputSize;
+
+		// Scales UV's such that the X axis is always 1.0
+		FVector2f InputHeightOverWidth = FVector2f(1.0f, static_cast<float>(InputSize.Y) / static_cast<float>(InputSize.X));
+
+		FIntVector NumGroups = FComputeShaderUtils::GetGroupCount(InputSize, FIntPoint(8, 8));
+		int32 NumActiveChannels = Parameters.bPerChannelCalculateDF.CountSetBits();
+
+		FIntVector4 ChannelMap(ForceInitToZero);
+		{
+			int32 ActiveChannelIdx = 0;
+			for (int32 ChannelIdx = 0; ChannelIdx < UE::GeometryMask::Private::MaxNumChannels; ++ChannelIdx)
+			{
+				for (int32 ChannelMaskIdx = ActiveChannelIdx; ChannelMaskIdx < Parameters.bPerChannelCalculateDF.Num(); ++ChannelMaskIdx)
+				{
+					if (Parameters.bPerChannelCalculateDF[ChannelMaskIdx])
+					{
+						ActiveChannelIdx = ChannelMaskIdx;
+						break;
+					}
+				}
+
+				ChannelMap[ChannelIdx] = ActiveChannelIdx;
+				++ActiveChannelIdx;
+			}
+		}
+
+		// Used to calculate step size per-channel
+		FIntVector4 SizeFromRadius;
+		for (int32 ChannelIdx = 0; ChannelIdx < UE::GeometryMask::Private::MaxNumChannels; ++ChannelIdx)
+		{
+			SizeFromRadius[ChannelIdx] = Parameters.PerChannelRadius[ChannelIdx];
+		}
+
+		FIntPoint BufferSize = InputSize;
+		BufferSize.X *= NumActiveChannels;
+		
+		const uint32 BufferLength = BufferSize.X * BufferSize.Y;
+
+		FGlobalShaderMap* GlobalShaderMap = GetGlobalShaderMap(GMaxRHIFeatureLevel);
+
+		TShaderRef<FGeometryMaskJFInitCSBase> InitComputeShader;
+		TShaderRef<FGeometryMaskJFStepCSBase> StepComputeShader;
+		TShaderRef<FGeometryMaskJFtoDFCSBase> OutputComputeShader;
+		
+		{
+			TShaderMapRef<FGeometryMaskJFInitCS1> InitComputeShader1(GlobalShaderMap);
+			TShaderMapRef<FGeometryMaskJFInitCS2> InitComputeShader2(GlobalShaderMap);
+			TShaderMapRef<FGeometryMaskJFInitCS3> InitComputeShader3(GlobalShaderMap);
+			TShaderMapRef<FGeometryMaskJFInitCS4> InitComputeShader4(GlobalShaderMap);
+
+			TShaderMapRef<FGeometryMaskJFStepCS1> StepComputeShader1(GlobalShaderMap);
+			TShaderMapRef<FGeometryMaskJFStepCS2> StepComputeShader2(GlobalShaderMap);
+			TShaderMapRef<FGeometryMaskJFStepCS3> StepComputeShader3(GlobalShaderMap);
+			TShaderMapRef<FGeometryMaskJFStepCS4> StepComputeShader4(GlobalShaderMap);
+
+			TShaderMapRef<FGeometryMaskJFtoDFCS1> OutputComputeShader1(GlobalShaderMap);
+			TShaderMapRef<FGeometryMaskJFtoDFCS2> OutputComputeShader2(GlobalShaderMap);
+			TShaderMapRef<FGeometryMaskJFtoDFCS3> OutputComputeShader3(GlobalShaderMap);
+			TShaderMapRef<FGeometryMaskJFtoDFCS4> OutputComputeShader4(GlobalShaderMap);
+
+			if (NumActiveChannels == 1)
+			{
+				InitComputeShader = InitComputeShader1;
+				StepComputeShader = StepComputeShader1;
+				OutputComputeShader = OutputComputeShader1;
+			}
+			else if (NumActiveChannels == 2)
+			{
+				InitComputeShader = InitComputeShader2;
+				StepComputeShader = StepComputeShader2;
+				OutputComputeShader = OutputComputeShader2;
+			}
+			else if (NumActiveChannels == 3)
+			{
+				InitComputeShader = InitComputeShader3;
+				StepComputeShader = StepComputeShader3;
+				OutputComputeShader = OutputComputeShader3;
+			}
+			else if(NumActiveChannels == 4)
+			{
+				InitComputeShader = InitComputeShader4;
+				StepComputeShader = StepComputeShader4;
+				OutputComputeShader = OutputComputeShader4;
+			}
+		}
+
+		FRDGTextureRef InputTexture = InTexture->GetRenderTargetTexture(GraphBuilder);
+		FRDGTextureSRVRef InputTexture_SRV = GraphBuilder.CreateSRV(FRDGTextureSRVDesc(InputTexture));
+
+		FRDGBufferDesc InitOutputBufferDesc =
+			FRDGBufferDesc::CreateBufferDesc(
+				sizeof(FVector4f),
+				BufferLength);
+
+		FRDGBufferRef InitOutputBuffer = GraphBuilder.CreateBuffer(InitOutputBufferDesc, TEXT("GeometryMaskJFInit.OutputBuffer"));
+		
+		FVector4f* BufferData = GraphBuilder.AllocPODArray<FVector4f>(BufferLength);
+		GraphBuilder.QueueBufferUpload(InitOutputBuffer, BufferData, sizeof(FVector4f) * BufferLength, ERDGInitialDataFlags::NoCopy);
+
+		// 1. Init from input binary-ish texture
+		{
+			FRDGBufferUAVRef OutputBuffer_UAV = GraphBuilder.CreateUAV(InitOutputBuffer, EPixelFormat::PF_FloatRGBA);
+
+			{
+				RDG_EVENT_SCOPE(GraphBuilder, "GeometryMaskJFInit");
+				
+				FGeometryMaskJFInitCSBase::FParameters* PassParameters = GraphBuilder.AllocParameters<FGeometryMaskJFInitCSBase::FParameters>();
+				{
+					PassParameters->InputDimensions = InputSize;
+					PassParameters->OneOverInputDimensions = OneOverInputSize;
+					PassParameters->UVRatioAdjustment = InputHeightOverWidth;
+					PassParameters->ChannelMap = ChannelMap;
+					PassParameters->InputTexture = InputTexture_SRV;
+					PassParameters->InputSampler = TStaticSamplerState<SF_Bilinear>::GetRHI();
+					PassParameters->OutputBuffer = OutputBuffer_UAV;
+				}
+
+				ClearUnusedGraphResources(InitComputeShader, PassParameters);
+			
+				FComputeShaderUtils::AddPass(
+					GraphBuilder,
+					RDG_EVENT_NAME("Init"),
+					ERDGPassFlags::Compute,
+					InitComputeShader,
+					PassParameters,
+					NumGroups);
+			}
+		}
+
+		FRDGBufferSRVRef StepOutputBuffer_SRV = nullptr;
+
+		// 2. JumpFlood steps
+		{
+			FRDGBufferDesc StepIntermediateBufferDescA =
+				FRDGBufferDesc::CreateBufferDesc(
+					sizeof(FVector4f),
+					BufferLength);
+			StepIntermediateBufferDescA.Usage = EBufferUsageFlags::UnorderedAccess;
+
+			FRDGBufferDesc StepIntermediateBufferDescB =
+				FRDGBufferDesc::CreateBufferDesc(
+					sizeof(FVector4f),
+					BufferLength);
+			StepIntermediateBufferDescB.Usage = EBufferUsageFlags::UnorderedAccess;
+
+			FRDGBufferRef StepIntermediateBufferA = GraphBuilder.CreateBuffer(StepIntermediateBufferDescA, TEXT("GeometryMaskJFStep.IntermediateBufferA"));
+			FRDGBufferRef StepIntermediateBufferB = GraphBuilder.CreateBuffer(StepIntermediateBufferDescB, TEXT("GeometryMaskJFStep.IntermediateBufferB"));
+
+			// BufferData = GraphBuilder.AllocPODArray<FVector4f>(BufferLength);
+			GraphBuilder.QueueBufferUpload(StepIntermediateBufferA, BufferData, sizeof(FVector4f) * BufferLength, ERDGInitialDataFlags::NoCopy);
+			
+			// BufferData = GraphBuilder.AllocPODArray<FVector4f>(BufferLength);
+			GraphBuilder.QueueBufferUpload(StepIntermediateBufferB, BufferData, sizeof(FVector4f) * BufferLength, ERDGInitialDataFlags::NoCopy);
+
+			// Copy init input to intermediate A
+			{
+				RDG_EVENT_SCOPE(GraphBuilder, "GeometryMaskJFStep.CopyInitToStep");
+
+				AddCopyBufferPass(GraphBuilder, StepIntermediateBufferA, InitOutputBuffer);
+			}
+			
+			FRDGBufferUAVRef StepIntermediateBufferA_UAV = GraphBuilder.CreateUAV(StepIntermediateBufferA, EPixelFormat::PF_FloatRGBA);
+			FRDGBufferUAVRef StepIntermediateBufferB_UAV = GraphBuilder.CreateUAV(StepIntermediateBufferB, EPixelFormat::PF_FloatRGBA);
+
+			FRDGBufferSRVRef StepIntermediateBufferA_SRV = GraphBuilder.CreateSRV(StepIntermediateBufferA, EPixelFormat::PF_FloatRGBA);
+			FRDGBufferSRVRef StepIntermediateBufferB_SRV = GraphBuilder.CreateSRV(StepIntermediateBufferB, EPixelFormat::PF_FloatRGBA);
+
+			StepOutputBuffer_SRV = StepIntermediateBufferA_SRV;
+
+			{
+				FIntVector4 StepCount = UE::GeometryMask::Private::CalculateStepCount(SizeFromRadius);
+				FIntVector4 ChannelMask(ForceInitToZero);
+				int32 MaxStepCount = FMath::Max(FMath::Max(FMath::Max(StepCount.X, StepCount.Y), StepCount.Z), StepCount.W);
+				int32 BufferIdx = 0;
+				
+				for (int32 StepIdx = 0; StepIdx < MaxStepCount; ++StepIdx)
+				{
+					RDG_EVENT_SCOPE(GraphBuilder, "GeometryMaskJFStep");
+			
+					FGeometryMaskJFStepCSBase::FParameters* PassParameters = GraphBuilder.AllocParameters<FGeometryMaskJFStepCSBase::FParameters>();
+					{
+						PassParameters->InputDimensions = InputSize;
+						PassParameters->OneOverInputDimensions = OneOverInputSize;
+						PassParameters->UVRatioAdjustment = InputHeightOverWidth;
+
+						PassParameters->StepSize = UE::GeometryMask::Private::CalculateStepSize(StepIdx, StepCount, ChannelMask);
+						PassParameters->ChannelMask = ChannelMask;
+						
+						PassParameters->InputBuffer = BufferIdx > 0
+							? StepIntermediateBufferB_SRV
+							: StepIntermediateBufferA_SRV;
+
+						PassParameters->OutputBuffer = BufferIdx > 0
+							? StepIntermediateBufferA_UAV
+							: StepIntermediateBufferB_UAV;
+					}
+
+					ClearUnusedGraphResources(StepComputeShader, PassParameters);
+			
+					FComputeShaderUtils::AddPass(
+						GraphBuilder,
+						RDG_EVENT_NAME("Step"),
+						ERDGPassFlags::Compute,
+						StepComputeShader,
+						PassParameters,
+						NumGroups);
+
+					BufferIdx = (BufferIdx + 1) % 2;
+				}
+
+				StepOutputBuffer_SRV = BufferIdx > 0
+					? StepIntermediateBufferB_SRV
+					: StepIntermediateBufferA_SRV;
+			}
+		}
+
+		// 3. JF to DF
+		{
+			FRDGTextureDesc OutputTextureDesc(
+				FRDGTextureDesc::Create2D(
+					InputTexture->Desc.Extent,
+					InputTexture->Desc.Format,
+					FClearValueBinding::Green,
+					TexCreate_ShaderResource | TexCreate_UAV));
+
+			FRDGTextureRef OutputTexture = GraphBuilder.CreateTexture(OutputTextureDesc, TEXT("GeometryMaskJFtoDF.OutputTexture"));
+			FRDGTextureUAVRef OutputTexture_UAV = GraphBuilder.CreateUAV(FRDGTextureUAVDesc(OutputTexture));
+
+			{
+				RDG_EVENT_SCOPE(GraphBuilder, "GeometryMaskJFtoDF");
+			
+				FGeometryMaskJFtoDFCSBase::FParameters* PassParameters = GraphBuilder.AllocParameters<FGeometryMaskJFtoDFCSBase::FParameters>();
+				{
+					FVector4f StepDistanceMultipliers;
+					for (int32 ChannelIdx = 0; ChannelIdx < UE::GeometryMask::Private::MaxNumChannels; ++ChannelIdx)
+					{
+						const float Radius = Parameters.bPerChannelCalculateDF[ChannelIdx] ? FMath::Max(1, Parameters.PerChannelRadius[ChannelIdx]) : InputSize.GetMax();
+						StepDistanceMultipliers[ChannelIdx] = static_cast<float>(InputSize.GetMax()) / Radius; 
+					}
+					
+					PassParameters->InputDimensions = InputSize;
+					PassParameters->OneOverInputDimensions = OneOverInputSize;
+					PassParameters->UVRatioAdjustment = InputHeightOverWidth;
+					PassParameters->ChannelMap = ChannelMap;
+					PassParameters->StepDistanceMultipliers = StepDistanceMultipliers;
+					PassParameters->OriginalInputTexture = InputTexture_SRV;
+					PassParameters->OriginalInputSampler = TStaticSamplerState<SF_Bilinear>::GetRHI();
+					PassParameters->InputBuffer = StepOutputBuffer_SRV;
+					PassParameters->OutputTexture = OutputTexture_UAV;
+				}
+
+				ClearUnusedGraphResources(OutputComputeShader, PassParameters);
+			
+				FComputeShaderUtils::AddPass(
+					GraphBuilder,
+					RDG_EVENT_NAME("CopyToDF"),
+					ERDGPassFlags::Compute,
+					OutputComputeShader,
+					PassParameters,
+					NumGroups);
+			}
+
+			// Copy back to RenderTarget
+			{
+				RDG_EVENT_SCOPE(GraphBuilder, "GeometryMaskDF.PostCopy");
+
+				AddCopyTexturePass(GraphBuilder, OutputTexture, InputTexture, FRHICopyTextureInfo());
+			}
+		}
+	}
+
+	GraphBuilder.Execute();
+}
diff --git a/Engine/Plugins/Experimental/GeometryMask/Source/GeometryMask/Private/Shaders/GeometryMaskPostProcess_DistanceField.h b/Engine/Plugins/Experimental/GeometryMask/Source/GeometryMask/Private/Shaders/GeometryMaskPostProcess_DistanceField.h
new file mode 100644
--- /dev/null
+++ b/Engine/Plugins/Experimental/GeometryMask/Source/GeometryMask/Private/Shaders/GeometryMaskPostProcess_DistanceField.h
@@ -0,0 +1,35 @@
+﻿// Copyright Epic Games, Inc. All Rights Reserved.
+
+#pragma once
+
+#include "Containers/Array.h"
+#include "Containers/BitArray.h"
+#include "Containers/ContainerAllocationPolicies.h"
+#include "GeometryMaskPostProcess.h"
+#include "Templates/SharedPointer.h"
+
+class FRenderTarget;
+
+struct FGeometryMaskPostProcessParameters_DistanceField
+{
+	FGeometryMaskPostProcessParameters_DistanceField();
+	
+	TBitArray<TInlineAllocator<4>> bPerChannelCalculateDF;
+	TArray<double, TInlineAllocator<4>> PerChannelRadius;
+};
+
+/** */
+class FGeometryMaskPostProcess_DistanceField
+	: public TGeometryMaskPostProcess<FGeometryMaskPostProcessParameters_DistanceField>
+	, public TSharedFromThis<FGeometryMaskPostProcess_DistanceField>
+{
+	using Super = TGeometryMaskPostProcess<FGeometryMaskPostProcessParameters_DistanceField>;
+	
+public:
+	explicit FGeometryMaskPostProcess_DistanceField(const FGeometryMaskPostProcessParameters_DistanceField& InParameters);
+
+	virtual void Execute(FRenderTarget* InTexture) override;
+
+protected:
+	virtual void Execute_RenderThread(FRHICommandListImmediate& InRHICmdList, FRenderTarget* InTexture) override;
+};
diff --git a/Engine/Plugins/Experimental/GeometryMask/Source/GeometryMask/Public/GeometryMaskCanvas.h b/Engine/Plugins/Experimental/GeometryMask/Source/GeometryMask/Public/GeometryMaskCanvas.h
new file mode 100644
--- /dev/null
+++ b/Engine/Plugins/Experimental/GeometryMask/Source/GeometryMask/Public/GeometryMaskCanvas.h
@@ -0,0 +1,154 @@
+﻿// Copyright Epic Games, Inc. All Rights Reserved.
+
+#pragma once
+
+#include "IGeometryMaskReadInterface.h"
+#include "SceneView.h"
+#include "UObject/Object.h"
+#include "UObject/WeakInterfacePtr.h"
+
+#include "GeometryMaskCanvas.generated.h"
+
+class FGeometryMaskPostProcess_Blur;
+class IGeometryMaskWriteInterface;
+class UCanvas;
+class UGeometryMaskCanvasResource;
+class UGeometryMaskPrimitive;
+class UGeometryMaskSubsystem;
+
+/** Called when Writers becomes non-empty. */
+using FOnGeometryMaskCanvasActivated = TDelegate<void()>;
+
+/** Called when Writers becomes empty. */
+using FOnGeometryMaskCanvasDeactivated = TDelegate<void()>;
+
+/** A uniquely identified Canvas. */
+UCLASS(BlueprintType, Transient)
+class GEOMETRYMASK_API UGeometryMaskCanvas : public UObject
+{
+	GENERATED_BODY()
+
+public:
+	virtual void BeginDestroy() override;
+	
+	/** Returns all writers. */
+    const TArray<TWeakInterfacePtr<IGeometryMaskWriteInterface>>& GetWriters() const;
+
+	/** Adds a Mask Writer to the canvas. */
+	void AddWriter(const TScriptInterface<IGeometryMaskWriteInterface>& InWriter);
+
+	/** Appends multiple writers to the canvas. */
+	void AddWriters(const TArray<TScriptInterface<IGeometryMaskWriteInterface>>& InWriters);
+
+	/** Remove a writer from this canvas. */
+	void RemoveWriter(const TScriptInterface<IGeometryMaskWriteInterface>& InWriter);
+
+	/** Returns true if this is the default/blank canvas. */
+	bool IsDefaultCanvas() const { return CanvasName == NAME_None; }
+
+	/** Get the underlying render target. */
+	UFUNCTION(BlueprintCallable, BlueprintPure, Category = "Rendering")
+	UTextureRenderTarget2D* GetTexture() const;
+
+	/** Whether blur is applied or not. */
+	bool IsBlurApplied() const;
+
+	/** Sets whether blur is applied or not. */
+	void SetApplyBlur(const bool bInValue);
+
+	/** Gets the current blur strength, if applicable. */
+	double GetBlurStrength() const;
+
+	/** Sets the blur strength. 0.0 will disable blur. */
+	void SetBlurStrength(const double InValue);
+
+	/** Whether feathering is applied or not. */
+	bool IsFeatherApplied() const;
+
+	/** Sets whether feathering is applied or not. */
+	void SetApplyFeather(const bool bInValue);
+
+	/** Gets the current outer feather radius, in pixels. */
+	int32 GetOuterFeatherRadius() const;
+
+	/** Sets the outer feather radius, in pixels. */
+	void SetOuterFeatherRadius(const int32 InValue);
+
+	/** Gets the current inner feather radius, in pixels. */
+	int32 GetInnerFeatherRadius() const;
+
+	/** Sets the inner feather radius, in pixels. */
+	void SetInnerFeatherRadius(const int32 InValue);
+
+	/** Get the color channel to write to in the texture. */
+	UFUNCTION(BlueprintCallable, BlueprintPure, Category = "Rendering")
+	EGeometryMaskColorChannel GetColorChannel() const;
+
+	/** Called when Writers becomes non-empty. */
+	FOnGeometryMaskCanvasActivated& OnActivated() { return OnActivatedDelegate; }
+
+	/** Called when Writers becomes empty. */
+	FOnGeometryMaskCanvasDeactivated& OnDeactivated() { return OnDeactivatedDelegate; }
+
+	/** Updates the canvas, intended to be called every frame. */
+	void Update(UWorld* InWorld, FSceneView& InView);
+	
+	const UGeometryMaskCanvasResource* GetResource() const { return CanvasResource; }
+	void AssignResource(UGeometryMaskCanvasResource* InResource, EGeometryMaskColorChannel InColorChannel);
+	void FreeResource();
+
+private:
+	/** Sorts writers by various criteria for proper rendering order. */
+	void SortWriters();
+	
+	/** Remove all invalid/stale writers. */
+	void RemoveInvalidWriters();
+
+	/** Draws all writers to the canvas. */
+	void OnDrawToCanvas(FCanvas* InCanvas);
+
+	/** Updates relevant shader parameters for the given color channel. */
+	void UpdateRenderParameters();
+	
+private:
+	// Allow the subsystem to set the CanvasName
+	friend class UGeometryMaskSubsystem;
+	
+	FOnGeometryMaskCanvasActivated OnActivatedDelegate;
+	FOnGeometryMaskCanvasDeactivated OnDeactivatedDelegate;
+
+	/** Uniquely identifies this canvas. */
+	UPROPERTY(VisibleAnywhere, BlueprintReadOnly, Getter = "Auto", Category = "Canvas", meta = (AllowPrivateAccess = "true"))
+	FName CanvasName;
+
+	/** Optional Blur Toggle. */
+	UPROPERTY(EditAnywhere, BlueprintReadWrite, Getter = "IsBlurApplied", Setter = "SetApplyBlur", Category = "Canvas", meta = (AllowPrivateAccess = "true"))
+	bool bApplyBlur = false;
+	
+	/** Optional Blur Strength. */
+	UPROPERTY(EditAnywhere, BlueprintReadWrite, Getter, Setter, Category = "Canvas", meta = (AllowPrivateAccess = "true", EditCondition = "bApplyBlur", EditConditionHides, ClampMin = 0.0))
+	double BlurStrength = 16;
+
+	/** Optional Feather Toggle. */
+	UPROPERTY(EditAnywhere, BlueprintReadWrite, Getter = "IsFeatherApplied", Setter = "SetApplyFeather", Category = "Canvas", meta = (AllowPrivateAccess = "true"))
+	bool bApplyFeather = false;
+	
+	/** Optional Outer Feather Radius. */
+	UPROPERTY(EditAnywhere, BlueprintReadWrite, Getter, Setter, Category = "Canvas", meta = (AllowPrivateAccess = "true", EditCondition = "bApplyFeather", EditConditionHides, ClampMin = 0))
+	int32 OuterFeatherRadius = 16;
+
+	/** Optional Inner Feather Radius. */
+	UPROPERTY(EditAnywhere, BlueprintReadWrite, Getter, Setter, Category = "Canvas", meta = (AllowPrivateAccess = "true", EditCondition = "bApplyFeather", EditConditionHides, ClampMin = 0))
+	int32 InnerFeatherRadius = 16;
+
+	/** Canvas GPU resource to use. */
+	UPROPERTY(Transient)
+	TObjectPtr<UGeometryMaskCanvasResource> CanvasResource;
+
+	/** Color Channel being written to in the texture. */
+	UPROPERTY(Transient, Getter)
+	EGeometryMaskColorChannel ColorChannel = EGeometryMaskColorChannel::None;
+
+	/** List of objects that write to this canvas. */
+	TArray<TWeakInterfacePtr<IGeometryMaskWriteInterface>> Writers;
+};
diff --git a/Engine/Plugins/Experimental/GeometryMask/Source/GeometryMask/Public/GeometryMaskCanvasActor.h b/Engine/Plugins/Experimental/GeometryMask/Source/GeometryMask/Public/GeometryMaskCanvasActor.h
new file mode 100644
--- /dev/null
+++ b/Engine/Plugins/Experimental/GeometryMask/Source/GeometryMask/Public/GeometryMaskCanvasActor.h
@@ -0,0 +1,54 @@
+﻿// Copyright Epic Games, Inc. All Rights Reserved.
+
+#pragma once
+
+#include "GameFramework/Actor.h"
+
+#include "GeometryMaskCanvasActor.generated.h"
+
+class UGeometryMaskCanvas;
+class IGeometryMaskWriteInterface;
+class UGeometryMaskCaptureComponent;
+class UGeometryMaskWriteMeshComponent;
+
+/** Wraps a GeometryMaskCanvas, and discovers/registers writers. */
+UCLASS(BlueprintType)
+class GEOMETRYMASK_API AGeometryMaskCanvasActor
+	: public AActor
+{
+	GENERATED_BODY()
+
+public:
+	AGeometryMaskCanvasActor(const FObjectInitializer& ObjectInitializer);
+
+	/** Identifies the referenced Canvas. */
+	UPROPERTY(EditAnywhere, BlueprintReadOnly, Category = "Canvas")
+	FName CanvasName;
+
+	/** Returns the Canvas Texture. */
+	UFUNCTION(BlueprintCallable, Category = "Rendering")
+	UTextureRenderTarget2D* GetTexture();
+
+	virtual void BeginPlay() override;
+
+	virtual void PostLoad() override;
+
+#if WITH_EDITOR
+	virtual void RerunConstructionScripts() override;
+#endif
+
+protected:
+	/** Resolve/locate the canvas identified by CanvasName. */
+	bool TryResolveCanvas();
+
+	/** Find writers on child actors and add them to the referenced canvas. */
+	void FindWriters();
+
+	/** Reference to the Canvas used, identified by CanvasName. */
+	UPROPERTY(Transient, DuplicateTransient)
+	TObjectPtr<UGeometryMaskCanvas> Canvas;
+
+	/** List of objects that write to this canvas. */
+	UPROPERTY()
+	TArray<TScriptInterface<IGeometryMaskWriteInterface>> Writers;
+};
diff --git a/Engine/Plugins/Experimental/GeometryMask/Source/GeometryMask/Public/GeometryMaskCanvasResource.h b/Engine/Plugins/Experimental/GeometryMask/Source/GeometryMask/Public/GeometryMaskCanvasResource.h
new file mode 100644
--- /dev/null
+++ b/Engine/Plugins/Experimental/GeometryMask/Source/GeometryMask/Public/GeometryMaskCanvasResource.h
@@ -0,0 +1,90 @@
+﻿// Copyright Epic Games, Inc. All Rights Reserved.
+
+#pragma once
+
+#include "UObject/Object.h"
+
+#include "GeometryMaskCanvasResource.generated.h"
+
+class FGeometryMaskPostProcess_Blur;
+class FGeometryMaskPostProcess_DistanceField;
+class FGeometryMaskPostProcess_SobelTest;
+class FCanvas;
+class FSceneView;
+class IGeometryMaskPostProcess;
+class UCanvas;
+class UTextureRenderTarget2D;
+enum class EGeometryMaskColorChannel : uint8;
+
+using FOnGeometryMaskCanvasDraw = TMulticastDelegate<void(FCanvas*)>;
+
+/**  */
+UCLASS()
+class GEOMETRYMASK_API UGeometryMaskCanvasResource
+	: public UObject
+{
+	GENERATED_BODY()
+
+public:
+	/** (R, G, B, A) */
+	static constexpr int32 MaxNumChannels = 4;
+	static constexpr int32 MaxTextureSize = 8192;
+	
+	UGeometryMaskCanvasResource();
+
+	/** Will return the first available color channel without a canvas assigned. EGeometryMaskColorChannel::None if not available. */
+	const EGeometryMaskColorChannel GetNextAvailableColorChannel() const;
+
+	/** Requests usage of the given color channel for this resource. Will return true if successful. */
+	bool Checkout(const EGeometryMaskColorChannel InColorChannel, const FName InRequestingCanvasName);
+
+	/** Returns/frees the color channel associated with the given canvas name. Returns true if canvas name found. */
+	bool Checkin(const FName InRequestingCanvasName);
+
+	void UpdateViewportSize();
+
+	void SetViewportSize(const FIntPoint& InViewportSize);
+
+	/** Returns required viewport padding, in pixels - determined by certain effects. */
+	int32 GetViewportPadding() const;
+
+	UTextureRenderTarget2D* GetRenderTargetTexture();
+
+	FOnGeometryMaskCanvasDraw& OnDrawToCanvas() { return OnDrawToCanvasDelegate; }
+
+	void UpdateRenderParameters(EGeometryMaskColorChannel InColorChannel, bool bInApplyBlur, double InBlurStrength, bool bInApplyFeather, int32 InOuterFeatherRadius, int32 InInnerFeatherRadius);
+
+	/** Updates the canvas, intended to be called every frame. */
+	void Update(UWorld* InWorld, FSceneView& InView);
+
+private:
+	/** Draws all writers to the canvas. */
+	void Draw(UWorld* InWorld, FSceneView& InView);
+
+private:
+	UPROPERTY(NoClear, EditFixedSize, meta = (EditFixedOrder))
+	TMap<EGeometryMaskColorChannel, FName> DependentCanvasNames;
+
+	/** Underlying UCanvas object. */
+	UPROPERTY(Transient)
+	TObjectPtr<UCanvas> CanvasObject;
+
+	/** The default viewport size to use when it can't be resolved from the actual viewport. */
+	UPROPERTY(EditAnywhere, BlueprintReadWrite, Getter = "Auto", Setter, Category = "Rendering", meta = (AllowPrivateAccess = "true"))
+	FIntPoint ViewportSize = FIntPoint(1920, 1080);
+	
+	/** The underlying Render Target texture. */
+	UPROPERTY(DuplicateTransient)
+	TObjectPtr<UTextureRenderTarget2D> RenderTargetTexture;
+
+	FOnGeometryMaskCanvasDraw OnDrawToCanvasDelegate;
+
+	/** The last resolved ViewProjectionMatrix. */
+	FMatrix CachedViewProjectionMatrix;
+
+	bool bApplyBlur = false;
+	bool bApplyDF = false;
+
+	TSharedPtr<FGeometryMaskPostProcess_Blur> PostProcess_Blur;
+	TSharedPtr<FGeometryMaskPostProcess_DistanceField> PostProcess_DistanceField;
+};
diff --git a/Engine/Plugins/Experimental/GeometryMask/Source/GeometryMask/Public/GeometryMaskReadComponent.h b/Engine/Plugins/Experimental/GeometryMask/Source/GeometryMask/Public/GeometryMaskReadComponent.h
new file mode 100644
--- /dev/null
+++ b/Engine/Plugins/Experimental/GeometryMask/Source/GeometryMask/Public/GeometryMaskReadComponent.h
@@ -0,0 +1,40 @@
+﻿// Copyright Epic Games, Inc. All Rights Reserved.
+
+#pragma once
+
+#include "Components/ActorComponent.h"
+#include "GeometryMaskTypes.h"
+#include "IGeometryMaskReadInterface.h"
+
+#include "GeometryMaskReadComponent.generated.h"
+
+class UGeometryMaskCanvas;
+
+UCLASS(Blueprintable, BlueprintType, HideCategories=(Activation, Cooking, AssetUserData, Navigation), meta=(BlueprintSpawnableComponent))
+class GEOMETRYMASK_API UGeometryMaskReadComponent
+	: public UGeometryMaskCanvasReferenceComponentBase
+	, public IGeometryMaskReadInterface
+{
+	GENERATED_BODY()
+
+public:
+	//~ Begin IGeometryMaskReadInterface
+	virtual void SetParameters(FGeometryMaskReadParameters& InParameters) override;
+	virtual const FGeometryMaskReadParameters& GetParameters() const override
+	{
+		return Parameters;
+	}
+	virtual FOnGeometryMaskSetCanvasNativeDelegate& OnSetCanvas() override
+	{
+		return OnSetCanvasDelegate;
+	}
+	//~ End IGeometryMaskReadInterface
+
+protected:
+	//~ Begin UGeometryMaskCanvasReferenceComponentBase
+	virtual bool TryResolveCanvas() override;
+	//~ End UGeometryMaskCanvasReferenceComponentBase
+
+	UPROPERTY(EditAnywhere, BlueprintReadWrite, Setter="SetParameters", Getter="GetParameters", Category="Mask", meta=(ShowOnlyInnerProperties))
+	FGeometryMaskReadParameters Parameters;
+};
diff --git a/Engine/Plugins/Experimental/GeometryMask/Source/GeometryMask/Public/GeometryMaskSubsystem.h b/Engine/Plugins/Experimental/GeometryMask/Source/GeometryMask/Public/GeometryMaskSubsystem.h
new file mode 100644
--- /dev/null
+++ b/Engine/Plugins/Experimental/GeometryMask/Source/GeometryMask/Public/GeometryMaskSubsystem.h
@@ -0,0 +1,65 @@
+﻿// Copyright Epic Games, Inc. All Rights Reserved.
+
+#pragma once
+
+#include "Engine/World.h"
+#include "GeometryMaskCanvas.h"
+#include "GeometryMaskCanvasResource.h"
+#include "SceneView.h"
+#include "Subsystems/EngineSubsystem.h"
+#include "Subsystems/WorldSubsystem.h"
+
+#include "GeometryMaskSubsystem.generated.h"
+
+using FOnGeometryMaskCanvasCreated = TMulticastDelegate<void(const UGeometryMaskCanvas*)>;
+using FOnGeometryMaskResourceCreated = TMulticastDelegate<void(const UGeometryMaskCanvasResource*)>;
+
+/** Maintains the registered named canvases. */
+UCLASS(BlueprintType)
+class GEOMETRYMASK_API UGeometryMaskSubsystem
+	: public UEngineSubsystem
+{
+	GENERATED_BODY()
+
+public:	
+	/** Retrieves a Canvas, uniquely identified by it's name. */
+	UFUNCTION(BlueprintCallable, Category = "Canvas")
+	UGeometryMaskCanvas* GetNamedCanvas(FName InName);
+
+	/** Returns all registered canvas names. */
+	UFUNCTION(BlueprintCallable, Category = "Canvas")
+	static TArray<FName> GetCanvasNames();
+
+	const TArray<TObjectPtr<UGeometryMaskCanvasResource>>& GetCanvasResources() const;
+
+	void Update(UWorld* InWorld, FSceneViewFamily& InViewFamily);
+
+	/** Remove all canvases without any Readers or Writers. Return the number of canvases removed. */
+	int32 RemoveWithoutWriters();
+
+	/** Called when a new canvas is created due to a unique name being requested. */
+	FOnGeometryMaskCanvasCreated& OnGeometryMaskCanvasCreated() { return OnGeometryMaskCanvasCreatedDelegate; }
+
+	/** Called when a new canvas resource is created. */
+	FOnGeometryMaskResourceCreated& OnGeometryMaskResourceCreated() { return OnGeometryMaskResourceCreatedDelegate; }
+
+private:
+	/** Find and assign the next available resource to the given canvas. */
+	void AssignResourceToCanvas(UGeometryMaskCanvas* InCanvas);
+
+	void OnCanvasActivated(UGeometryMaskCanvas* InCanvas);
+	void OnCanvasDeactivated(UGeometryMaskCanvas* InCanvas);
+
+private:
+	friend class UGeometryMaskWorldSubsystem;
+
+	FOnGeometryMaskCanvasCreated OnGeometryMaskCanvasCreatedDelegate;
+	FOnGeometryMaskResourceCreated OnGeometryMaskResourceCreatedDelegate;
+
+	UPROPERTY()
+	TMap<FName, TObjectPtr<UGeometryMaskCanvas>> NamedCanvases;
+
+	/** Pool of GPU/Texture resources used by the canvases. */
+	UPROPERTY(Getter)
+	TArray<TObjectPtr<UGeometryMaskCanvasResource>> CanvasResources;
+};
diff --git a/Engine/Plugins/Experimental/GeometryMask/Source/GeometryMask/Public/GeometryMaskTypes.h b/Engine/Plugins/Experimental/GeometryMask/Source/GeometryMask/Public/GeometryMaskTypes.h
new file mode 100644
--- /dev/null
+++ b/Engine/Plugins/Experimental/GeometryMask/Source/GeometryMask/Public/GeometryMaskTypes.h
@@ -0,0 +1,166 @@
+﻿// Copyright Epic Games, Inc. All Rights Reserved.
+
+#pragma once
+
+#include "Components/ActorComponent.h"
+#include "CoreTypes.h"
+
+#include "GeometryMaskTypes.generated.h"
+
+class UGeometryMaskCanvas;
+
+UENUM(BlueprintType)
+enum class EGeometryMaskColorChannel : uint8
+{
+	Red = 0,
+	Green = 1 << 0,
+	Blue = 1 << 1,
+	Alpha = 1 << 2,
+	None = 1 << 3,
+
+	Num
+};
+ENUM_CLASS_FLAGS(EGeometryMaskColorChannel)
+
+namespace UE::GeometryMask
+{
+	/** Use to convert between a multiplied vector and an EGeometryMaskColorChannel enum value. */
+	static TMap<FLinearColor, EGeometryMaskColorChannel> VectorToMaskChannelEnum =
+	{
+		{ {1, 0, 0, 0}, EGeometryMaskColorChannel::Red },
+		{ {0, 1, 0, 0}, EGeometryMaskColorChannel::Green },
+		{ {0, 0, 1, 0}, EGeometryMaskColorChannel::Blue },
+		{ {0, 0, 0, 1}, EGeometryMaskColorChannel::Alpha }
+	};
+	
+	/** Use to convert between a EGeometryMaskColorChannel enum value and a multiplied vector. */
+	static TMap<EGeometryMaskColorChannel, FLinearColor> MaskChannelEnumToVector =
+	{
+		{ EGeometryMaskColorChannel::Red, {1, 0, 0, 0} },
+		{ EGeometryMaskColorChannel::Green, {0, 1, 0, 0} },
+		{ EGeometryMaskColorChannel::Blue, {0, 0, 1, 0} },
+		{ EGeometryMaskColorChannel::Alpha, {0, 0, 0, 1} }
+	};
+
+	/** Use to convert between a EGeometryMaskColorChannel enum value and an FColor. */
+	static TMap<EGeometryMaskColorChannel, FColor> MaskChannelEnumToColor =
+	{
+		{ EGeometryMaskColorChannel::Red, {255, 0, 0, 0} },
+		{ EGeometryMaskColorChannel::Green, {0, 255, 0, 0} },
+		{ EGeometryMaskColorChannel::Blue, {0, 0, 255, 0} },
+		{ EGeometryMaskColorChannel::Alpha, {0, 0, 0, 255} }
+	};	
+
+	/** Use to convert between a EGeometryMaskColorChannel and it's value name (as string). */
+	static TMap<EGeometryMaskColorChannel, FString> MaskChannelEnumToString =
+	{
+		{ EGeometryMaskColorChannel::Red, TEXT("Red") },
+		{ EGeometryMaskColorChannel::Green, TEXT("Green") },
+		{ EGeometryMaskColorChannel::Blue, TEXT("Blue") },
+		{ EGeometryMaskColorChannel::Alpha, TEXT("Alpha") },
+		{ EGeometryMaskColorChannel::None, TEXT("None") },
+		{ EGeometryMaskColorChannel::Num, TEXT("Num/Max") }
+	};
+
+	GEOMETRYMASK_API EGeometryMaskColorChannel VectorToMaskChannel(const FLinearColor& InVector);
+
+	/** Returns a valid ColorChannel (R,G,B or A). */
+	GEOMETRYMASK_API EGeometryMaskColorChannel GetValidMaskChannel(EGeometryMaskColorChannel InColorChannel, bool bInIncludeAlpha = false);
+
+	GEOMETRYMASK_API FStringView ChannelToString(EGeometryMaskColorChannel InColorChannel);
+};
+
+UENUM(BlueprintType)
+enum class EGeometryMaskCompositeOperation : uint8
+{
+	Add,
+	Subtract,
+	Intersect,
+
+	Num
+};
+
+DECLARE_DYNAMIC_MULTICAST_DELEGATE_OneParam(FOnGeometryMaskSetCanvasDelegate, const UGeometryMaskCanvas*, InCanvas);
+using FOnGeometryMaskSetCanvasNativeDelegate = TMulticastDelegate<void(const UGeometryMaskCanvas* InCanvas)>;
+
+USTRUCT(BlueprintType)
+struct FGeometryMaskReadParameters
+{
+	GENERATED_BODY()
+
+public:
+	/** Specifies the GeometryMaskCanvas to read from. */
+	UPROPERTY(EditAnywhere, BlueprintReadWrite, Category = "Canvas")
+	FName CanvasName;
+
+	UPROPERTY(EditAnywhere, Category = "Read")
+	EGeometryMaskColorChannel ColorChannel = EGeometryMaskColorChannel::Red;
+
+	UPROPERTY(EditAnywhere, Category = "Read")
+	bool bInvert = false;
+};
+
+USTRUCT(BlueprintType)
+struct FGeometryMaskWriteParameters
+{
+	GENERATED_BODY()
+
+public:
+	/** Specifies the GeometryMaskCanvas to reference. */
+	UPROPERTY(EditAnywhere, BlueprintReadWrite, Category = "Canvas")
+	FName CanvasName;
+	
+	UPROPERTY(EditAnywhere, Category = "Write")
+	EGeometryMaskColorChannel ColorChannel = EGeometryMaskColorChannel::Red;
+
+	UPROPERTY(EditAnywhere, Category = "Write")
+	EGeometryMaskCompositeOperation OperationType = EGeometryMaskCompositeOperation::Add;
+
+	UPROPERTY(EditAnywhere, Category = "Write")
+	bool bInvert = false;
+
+	/** Higher values draw above lower.  */
+	UPROPERTY(EditAnywhere, Category = "Write")
+	int32 Priority = 50;
+
+	/** The outer offset of the shapes extent in world units. Smoothly interpolates between the shape/inner radius and this. */
+	UPROPERTY(EditAnywhere, Category = "Shape", meta = (ClampMin = 0, UIMin = 0))
+	double OuterRadius = 0.0;
+
+	/** The inner offset of the shapes extent in world units. Smoothly interpolates between the shape/outer radius and this. */
+	UPROPERTY(EditAnywhere, Category = "Shape", meta = (ClampMin = 0, UIMin = 0))
+	double InnerRadius = 0.0;
+};
+
+UCLASS(Abstract)
+class GEOMETRYMASK_API UGeometryMaskCanvasReferenceComponentBase
+	: public UActorComponent
+{
+	GENERATED_BODY()
+
+public:
+	/** Implement to perform an operation with the provided canvas. */
+	UFUNCTION(BlueprintImplementableEvent, CallInEditor, Category = "Canvas", meta = (DisplayName = "Canvas Changed"))
+	void ReceiveSetCanvas(const UGeometryMaskCanvas* InCanvas);
+
+	/** Returns the Canvas Texture. */
+	UFUNCTION(BlueprintCallable, Category="Rendering")
+	UTextureRenderTarget2D* GetTexture();
+
+protected:
+	virtual void BeginPlay() override;
+	virtual void PostLoad() override;
+	virtual void OnRegister() override;
+
+	/** Override to inject CanvasName into TryResolveCanvas(FName). */
+	virtual bool TryResolveCanvas() PURE_VIRTUAL(UGeometryMaskCanvasReferenceComponentBase::TryResolveCanvas, return false; )
+	
+	bool TryResolveNamedCanvas(FName InCanvasName);
+	
+protected:
+	FOnGeometryMaskSetCanvasNativeDelegate OnSetCanvasDelegate;
+
+	/** Reference to the Canvas used, identified by CanvasName. */
+	UPROPERTY(Transient, DuplicateTransient, TextExportTransient)
+	TWeakObjectPtr<UGeometryMaskCanvas> CanvasWeak;
+};
diff --git a/Engine/Plugins/Experimental/GeometryMask/Source/GeometryMask/Public/GeometryMaskWorldSubsystem.h b/Engine/Plugins/Experimental/GeometryMask/Source/GeometryMask/Public/GeometryMaskWorldSubsystem.h
new file mode 100644
--- /dev/null
+++ b/Engine/Plugins/Experimental/GeometryMask/Source/GeometryMask/Public/GeometryMaskWorldSubsystem.h
@@ -0,0 +1,26 @@
+﻿// Copyright Epic Games, Inc. All Rights Reserved.
+
+#pragma once
+
+#include "Subsystems/WorldSubsystem.h"
+#include "UnrealClient.h"
+
+#include "GeometryMaskWorldSubsystem.generated.h"
+
+class FGeometryMaskSceneViewExtension;
+
+/** Updates the canvases. */
+UCLASS()
+class UGeometryMaskWorldSubsystem
+	: public UWorldSubsystem
+{
+	GENERATED_BODY()
+
+protected:
+	// ~Begin USubsystem
+	virtual void Initialize(FSubsystemCollectionBase& Collection) override;
+	// ~End USubsystem
+
+private:
+	TSharedPtr<FGeometryMaskSceneViewExtension> GeometryMaskSceneViewExtension;
+};
diff --git a/Engine/Plugins/Experimental/GeometryMask/Source/GeometryMask/Public/GeometryMaskWriteComponent.h b/Engine/Plugins/Experimental/GeometryMask/Source/GeometryMask/Public/GeometryMaskWriteComponent.h
new file mode 100644
--- /dev/null
+++ b/Engine/Plugins/Experimental/GeometryMask/Source/GeometryMask/Public/GeometryMaskWriteComponent.h
@@ -0,0 +1,101 @@
+﻿// Copyright Epic Games, Inc. All Rights Reserved.
+
+#pragma once
+
+#include "Components/ActorComponent.h"
+#include "GeometryMaskTypes.h"
+#include "IGeometryMaskWriteInterface.h"
+
+#include "GeometryMaskWriteComponent.generated.h"
+
+class UDynamicMeshComponent;
+class UStaticMeshComponent;
+
+/** Mesh data for a single mesh, in local space. */
+USTRUCT()
+struct FGeometryMaskBatchElementData
+{
+	GENERATED_BODY()
+
+	UPROPERTY()
+	TArray<FVector3f> Vertices;
+
+	UPROPERTY()
+	TArray<uint32> Indices;
+
+	UPROPERTY()
+	int32 NumTriangles = 0;
+
+	/** Used to check cached state. */
+	UPROPERTY()
+	uint64 ChangeStamp = 0;
+
+	void Reserve(
+		const int32 InNumVertices,
+		const int32 InNumIndices,
+		const int32 InNumTriangles)
+	{
+		Vertices.Reserve(InNumVertices);
+		Indices.Reserve(InNumIndices);
+		NumTriangles = InNumTriangles;
+	}
+};
+
+UCLASS(BlueprintType, HideCategories=(Activation, Cooking, AssetUserData, Navigation), meta=(BlueprintSpawnableComponent))
+class GEOMETRYMASK_API UGeometryMaskWriteMeshComponent
+	: public UGeometryMaskCanvasReferenceComponentBase
+	, public IGeometryMaskWriteInterface
+{
+	GENERATED_BODY()
+
+public:
+	//~ Begin IGeometryMaskWriteInterface
+	virtual void SetParameters(FGeometryMaskWriteParameters& InParameters) override;
+	virtual const FGeometryMaskWriteParameters& GetParameters() const override
+	{
+		return Parameters;
+	}
+	virtual void DrawToCanvas(FCanvas* InCanvas) override;
+	virtual FOnGeometryMaskSetCanvasNativeDelegate& OnSetCanvas() override
+	{
+		return OnSetCanvasDelegate;
+	}
+	//~ End IGeometryMaskWriteInterface
+
+protected:
+	//~ Begin UActorComponent
+	virtual void OnComponentDestroyed(bool bDestroyingHierarchy) override;
+	//~ End UActorComponent
+	
+	//~ Begin UGeometryMaskCanvasReferenceComponentBase
+    virtual bool TryResolveCanvas() override;
+    //~ End UGeometryMaskCanvasReferenceComponentBase
+    	
+	void UpdateCachedData();
+	void UpdateCachedStaticMeshData(TConstArrayView<UPrimitiveComponent*> InPrimitiveComponents);
+	void UpdateCachedDynamicMeshData(TConstArrayView<UPrimitiveComponent*> InPrimitiveComponents);
+
+#if WITH_EDITOR
+	void OnStaticMeshChanged(UStaticMeshComponent* InStaticMeshComponent);
+#endif
+	void OnDynamicMeshChanged(UDynamicMeshComponent* InDynamicMeshComponent);
+	
+	UPROPERTY(EditAnywhere, BlueprintReadWrite, Category="Mask", meta=(ShowOnlyInnerProperties))
+	FGeometryMaskWriteParameters Parameters;
+
+	/** Write to the mask even when this or it's owner is hidden. */
+	UPROPERTY(EditAnywhere, BlueprintReadWrite, Category="Mask")
+	bool bWriteWhenHidden = true;
+
+	UPROPERTY(Transient, DuplicateTransient)
+	int32 LastPrimitiveComponentCount;
+
+	UPROPERTY(DuplicateTransient)
+	TMap<FName, TWeakObjectPtr<USceneComponent>> CachedComponents;
+	
+	UPROPERTY(DuplicateTransient)
+	TMap<FName, FGeometryMaskBatchElementData> CachedMeshData;
+
+	/** To protect async-unsafe calls. */
+	FCriticalSection CanvasCS;
+};
diff --git a/Engine/Plugins/Experimental/GeometryMask/Source/GeometryMask/Public/IGeometryMaskModule.h b/Engine/Plugins/Experimental/GeometryMask/Source/GeometryMask/Public/IGeometryMaskModule.h
new file mode 100644
--- /dev/null
+++ b/Engine/Plugins/Experimental/GeometryMask/Source/GeometryMask/Public/IGeometryMaskModule.h
@@ -0,0 +1,17 @@
+// Copyright Epic Games, Inc. All Rights Reserved.
+
+#pragma once
+
+#include "Modules/ModuleInterface.h"
+
+class GEOMETRYMASK_API IGeometryMaskModule : public IModuleInterface
+{
+public:
+	/**
+	* Singleton-like access to this module's interface.  This is just for convenience!
+	* Beware of calling this during the shutdown phase, though.  Your module might have been unloaded already.
+	*
+	* @return Returns singleton instance, loading the module on demand if needed
+	*/
+	static IGeometryMaskModule& Get();
+};
diff --git a/Engine/Plugins/Experimental/GeometryMask/Source/GeometryMask/Public/IGeometryMaskReadInterface.h b/Engine/Plugins/Experimental/GeometryMask/Source/GeometryMask/Public/IGeometryMaskReadInterface.h
new file mode 100644
--- /dev/null
+++ b/Engine/Plugins/Experimental/GeometryMask/Source/GeometryMask/Public/IGeometryMaskReadInterface.h
@@ -0,0 +1,27 @@
+﻿// Copyright Epic Games, Inc. All Rights Reserved.
+
+#pragma once
+
+#include "GeometryMaskTypes.h"
+#include "UObject/Interface.h"
+#include "IGeometryMaskReadInterface.generated.h"
+
+class FCanvas;
+struct FGeometryMaskReadParameters;
+
+UINTERFACE(MinimalAPI, meta=(CannotImplementInterfaceInBlueprint))
+class UGeometryMaskReadInterface : public UInterface
+{
+	GENERATED_BODY()
+};
+
+/** Implement to Read from a host GeometryMaskCanvas. */
+class IGeometryMaskReadInterface
+{
+	GENERATED_BODY()
+
+public:
+	virtual const FGeometryMaskReadParameters& GetParameters() const = 0;
+	virtual void SetParameters(FGeometryMaskReadParameters& InParameters) = 0;
+	virtual FOnGeometryMaskSetCanvasNativeDelegate& OnSetCanvas() = 0;
+};
diff --git a/Engine/Plugins/Experimental/GeometryMask/Source/GeometryMask/Public/IGeometryMaskWriteInterface.h b/Engine/Plugins/Experimental/GeometryMask/Source/GeometryMask/Public/IGeometryMaskWriteInterface.h
new file mode 100644
--- /dev/null
+++ b/Engine/Plugins/Experimental/GeometryMask/Source/GeometryMask/Public/IGeometryMaskWriteInterface.h
@@ -0,0 +1,28 @@
+﻿// Copyright Epic Games, Inc. All Rights Reserved.
+
+#pragma once
+
+#include "GeometryMaskTypes.h"
+#include "UObject/Interface.h"
+#include "IGeometryMaskWriteInterface.generated.h"
+
+class FCanvas;
+struct FGeometryMaskWriteParameters;
+
+UINTERFACE(MinimalAPI, meta=(CannotImplementInterfaceInBlueprint))
+class UGeometryMaskWriteInterface : public UInterface
+{
+	GENERATED_BODY()
+};
+
+/** Implement to write to a host canvas. */
+class IGeometryMaskWriteInterface
+{
+	GENERATED_BODY()
+
+public:
+	virtual const FGeometryMaskWriteParameters& GetParameters() const = 0;
+	virtual void SetParameters(FGeometryMaskWriteParameters& InParameters) = 0;
+	virtual void DrawToCanvas(FCanvas* InCanvas) = 0;
+	virtual FOnGeometryMaskSetCanvasNativeDelegate& OnSetCanvas() = 0;
+};
diff --git a/Engine/Plugins/Experimental/GeometryMask/Source/GeometryMaskEditor/GeometryMaskEditor.Build.cs b/Engine/Plugins/Experimental/GeometryMask/Source/GeometryMaskEditor/GeometryMaskEditor.Build.cs
new file mode 100644
--- /dev/null
+++ b/Engine/Plugins/Experimental/GeometryMask/Source/GeometryMaskEditor/GeometryMaskEditor.Build.cs
@@ -0,0 +1,28 @@
+// Copyright Epic Games, Inc. All Rights Reserved.
+
+using UnrealBuildTool;
+
+public class GeometryMaskEditor : ModuleRules
+{
+	public GeometryMaskEditor(ReadOnlyTargetRules Target) : base(Target)
+	{
+		PCHUsage = ModuleRules.PCHUsageMode.UseExplicitOrSharedPCHs;
+
+		PublicDependencyModuleNames.AddRange(
+			new string[] {
+				"CoreUObject",
+                "GeometryMask",
+                "SlateCore",
+            });
+
+		PrivateDependencyModuleNames.AddRange(
+			new string[] {
+				"Core",
+				"Engine",
+				"InputCore",
+				"Slate",
+				"UnrealEd",
+				"WorkspaceMenuStructure",
+			});
+    }
+}
diff --git a/Engine/Plugins/Experimental/GeometryMask/Source/GeometryMaskEditor/Private/GeometryMaskEditorLog.cpp b/Engine/Plugins/Experimental/GeometryMask/Source/GeometryMaskEditor/Private/GeometryMaskEditorLog.cpp
new file mode 100644
--- /dev/null
+++ b/Engine/Plugins/Experimental/GeometryMask/Source/GeometryMaskEditor/Private/GeometryMaskEditorLog.cpp
@@ -0,0 +1,5 @@
+// Copyright Epic Games, Inc. All Rights Reserved.
+
+#include "GeometryMaskEditorLog.h"
+
+DEFINE_LOG_CATEGORY(LogGeometryMaskEditor);
diff --git a/Engine/Plugins/Experimental/GeometryMask/Source/GeometryMaskEditor/Private/GeometryMaskEditorLog.h b/Engine/Plugins/Experimental/GeometryMask/Source/GeometryMaskEditor/Private/GeometryMaskEditorLog.h
new file mode 100644
--- /dev/null
+++ b/Engine/Plugins/Experimental/GeometryMask/Source/GeometryMaskEditor/Private/GeometryMaskEditorLog.h
@@ -0,0 +1,7 @@
+// Copyright Epic Games, Inc. All Rights Reserved.
+
+#pragma once
+
+#include "Logging/LogMacros.h"
+
+DECLARE_LOG_CATEGORY_EXTERN(LogGeometryMaskEditor, VeryVerbose, All);
diff --git a/Engine/Plugins/Experimental/GeometryMask/Source/GeometryMaskEditor/Private/GeometryMaskEditorModule.cpp b/Engine/Plugins/Experimental/GeometryMask/Source/GeometryMaskEditor/Private/GeometryMaskEditorModule.cpp
new file mode 100644
--- /dev/null
+++ b/Engine/Plugins/Experimental/GeometryMask/Source/GeometryMaskEditor/Private/GeometryMaskEditorModule.cpp
@@ -0,0 +1,108 @@
+// Copyright Epic Games, Inc. All Rights Reserved.
+
+#include "GeometryMaskEditorModule.h"
+
+#include "Framework/Docking/TabManager.h"
+#include "Materials/Material.h"
+#include "Styling/SlateIconFinder.h"
+#include "ViewModels/GMECanvasListViewModel.h"
+#include "ViewModels/GMEResourceListViewModel.h"
+#include "Widgets/Docking/SDockTab.h"
+#include "Widgets/SGMECanvasList.h"
+#include "Widgets/SGMEResourceList.h"
+#include "WorkspaceMenuStructure.h"
+#include "WorkspaceMenuStructureModule.h"
+
+#define LOCTEXT_NAMESPACE "FGeometryMaskEditorModule"
+
+const FName FGeometryMaskEditorModule::VisualizerTabId = TEXT("GeometryMaskVisualizer");
+
+void FGeometryMaskEditorModule::StartupModule()
+{
+	const FText TabText = LOCTEXT("GeometryMaskVisualizerTabName", "Geometry Mask");
+	TSharedRef<FGlobalTabmanager> GlobalTabManager = FGlobalTabmanager::Get();
+
+	GlobalTabManager->RegisterNomadTabSpawner(
+			VisualizerTabId,
+			FOnSpawnTab::CreateLambda(
+				[this](const FSpawnTabArgs&)
+				{
+					TSharedRef<FGMECanvasListViewModel> CanvasListViewModel = FGMECanvasListViewModel::Create();
+					TSharedRef<FGMEResourceListViewModel> ResourceListViewModel = FGMEResourceListViewModel::Create();
+
+					TSharedRef<SDockTab> DockTab =
+						SNew(SDockTab)
+						.TabRole(ETabRole::NomadTab)
+						.Content()
+						[
+							SNew(SHorizontalBox)
+							+ SHorizontalBox::Slot()
+							.FillWidth(0.5)
+							[
+								SNew(SGMECanvasList, CanvasListViewModel)
+							]
+							
+							+ SHorizontalBox::Slot()
+							.FillWidth(0.5)
+							[
+								SNew(SGMEResourceList, ResourceListViewModel)
+							]
+						];
+
+					return DockTab;
+				}
+			)
+		)
+		.SetDisplayNameAttribute(TabText)
+		.SetDisplayName(TabText)
+		.SetGroup(WorkspaceMenu::GetMenuStructure().GetLevelEditorCategory())
+		.SetIcon(FSlateIconFinder::FindIconForClass(UMaterial::StaticClass()))
+		.SetAutoGenerateMenuEntry(false); // Only show via console command
+	
+	if (GIsEditor && !IsRunningCommandlet())
+	{
+		RegisterConsoleCommands();
+	}
+}
+
+void FGeometryMaskEditorModule::ShutdownModule()
+{
+	const TSharedRef<FGlobalTabmanager> GlobalTabManager = FGlobalTabmanager::Get();
+	GlobalTabManager->UnregisterNomadTabSpawner(VisualizerTabId);
+	VisualizerTabWeak.Reset();
+	
+	// Cleanup commands
+	UnregisterConsoleCommands();
+}
+
+void FGeometryMaskEditorModule::RegisterConsoleCommands()
+{
+	ConsoleCommands.Add(IConsoleManager::Get().RegisterConsoleCommand(
+		TEXT("GeometryMask.Visualize"),
+		TEXT("Displays a window listing all active GeometryMaskCanvas objects"),
+		FConsoleCommandWithArgsDelegate::CreateRaw(this, &FGeometryMaskEditorModule::ExecuteShowVisualizer),
+		ECVF_Default
+	));
+}
+
+void FGeometryMaskEditorModule::UnregisterConsoleCommands()
+{
+	for (IConsoleObject* Cmd : ConsoleCommands)
+	{
+		IConsoleManager::Get().UnregisterConsoleObject(Cmd);
+	}
+
+	ConsoleCommands.Empty();
+}
+
+void FGeometryMaskEditorModule::ExecuteShowVisualizer(const TArray<FString>& InArgs)
+{
+	if (!VisualizerTabWeak.IsValid())
+	{
+		VisualizerTabWeak = FGlobalTabmanager::Get()->TryInvokeTab(VisualizerTabId); 
+	}
+}
+
+#undef LOCTEXT_NAMESPACE
+
+IMPLEMENT_MODULE(FGeometryMaskEditorModule, GeometryMaskEditor)
diff --git a/Engine/Plugins/Experimental/GeometryMask/Source/GeometryMaskEditor/Private/GeometryMaskEditorModule.h b/Engine/Plugins/Experimental/GeometryMask/Source/GeometryMaskEditor/Private/GeometryMaskEditorModule.h
new file mode 100644
--- /dev/null
+++ b/Engine/Plugins/Experimental/GeometryMask/Source/GeometryMaskEditor/Private/GeometryMaskEditorModule.h
@@ -0,0 +1,39 @@
+﻿// Copyright Epic Games, Inc. All Rights Reserved.
+
+#pragma once
+
+#include "Containers/Array.h"
+#include "Containers/UnrealString.h"
+#include "Modules/ModuleInterface.h"
+#include "Templates/SharedPointer.h"
+#include "UObject/NameTypes.h"
+
+class IConsoleObject;
+class SDockTab;
+
+class FGeometryMaskEditorModule
+	: public IModuleInterface
+{
+public:
+	// ~Begin IModuleInterface
+	virtual void StartupModule() override;
+	virtual void ShutdownModule() override;
+	// ~End IModuleInterface
+
+private:
+	/** Called from StartupModule and sets up console commands for the plugin via IConsoleManager */
+	void RegisterConsoleCommands();
+	
+	/** Called from ShutdownModule and clears out previously registered console commands */
+	void UnregisterConsoleCommands();
+
+	void ExecuteShowVisualizer(const TArray<FString>& InArgs);
+
+private:
+	static const FName VisualizerTabId;
+	
+	/** References of registered console commands via IConsoleManager */
+	TArray<IConsoleObject*> ConsoleCommands;
+
+	TWeakPtr<SDockTab> VisualizerTabWeak;
+};
diff --git a/Engine/Plugins/Experimental/GeometryMask/Source/GeometryMaskEditor/Private/ViewModels/GMECanvasItemViewModel.cpp b/Engine/Plugins/Experimental/GeometryMask/Source/GeometryMaskEditor/Private/ViewModels/GMECanvasItemViewModel.cpp
new file mode 100644
--- /dev/null
+++ b/Engine/Plugins/Experimental/GeometryMask/Source/GeometryMaskEditor/Private/ViewModels/GMECanvasItemViewModel.cpp
@@ -0,0 +1,52 @@
+﻿// Copyright Epic Games, Inc. All Rights Reserved.
+
+#include "GMECanvasItemViewModel.h"
+
+#include "Engine/Texture.h"
+#include "Engine/TextureRenderTarget2D.h"
+#include "GeometryMaskCanvas.h"
+
+TSharedRef<FGMECanvasItemViewModel> FGMECanvasItemViewModel::Create(const TWeakObjectPtr<const UGeometryMaskCanvas>& InCanvas)
+{
+	TSharedRef<FGMECanvasItemViewModel> ViewModel = MakeShared<FGMECanvasItemViewModel>(FPrivateToken{}, InCanvas);
+
+	return ViewModel;
+}
+
+const UTexture* FGMECanvasItemViewModel::GetCanvasTexture() const
+{
+	if (CanvasTexture.IsValid())
+	{
+		return CanvasTexture.Get();
+	}
+	
+	return nullptr;
+}
+
+float FGMECanvasItemViewModel::GetMemoryUsage()
+{
+	if (const UTexture* Texture = GetCanvasTexture())
+	{
+		return (static_cast<float>(Texture->CalcTextureMemorySizeEnum(TMC_ResidentMips)) / (1024.0f * 1024.0f));
+	}
+	
+	return 0.0f;
+}
+
+FGMECanvasItemViewModel::FGMECanvasItemViewModel(FPrivateToken, const TWeakObjectPtr<const UGeometryMaskCanvas>& InCanvas)
+	: KnownReaderCount(0)
+	, KnownWriterCount(0)
+{
+	if (const UGeometryMaskCanvas* Canvas = InCanvas.Get())
+	{
+		CanvasName = Canvas->GetCanvasName();
+		ColorChannel = Canvas->GetColorChannel();
+		CanvasTexture = Canvas->GetTexture();
+		KnownWriterCount = Canvas->GetWriters().Num();
+	}
+}
+
+bool FGMECanvasItemViewModel::GetChildren(TArray<TSharedPtr<IGMETreeNodeViewModel>>& OutChildren)
+{
+	return false;
+}
diff --git a/Engine/Plugins/Experimental/GeometryMask/Source/GeometryMaskEditor/Private/ViewModels/GMECanvasItemViewModel.h b/Engine/Plugins/Experimental/GeometryMask/Source/GeometryMaskEditor/Private/ViewModels/GMECanvasItemViewModel.h
new file mode 100644
--- /dev/null
+++ b/Engine/Plugins/Experimental/GeometryMask/Source/GeometryMaskEditor/Private/ViewModels/GMECanvasItemViewModel.h
@@ -0,0 +1,48 @@
+﻿// Copyright Epic Games, Inc. All Rights Reserved.
+
+#pragma once
+
+#include "Containers/Array.h"
+#include "EditorUndoClient.h"
+#include "GeometryMaskCanvasResource.h"
+#include "GMEViewModelShared.h"
+#include "Templates/SharedPointer.h"
+#include "UObject/NameTypes.h"
+#include "UObject/WeakObjectPtr.h"
+
+class UGeometryMaskCanvas;
+class UTexture;
+
+class FGMECanvasItemViewModel
+	: public TSharedFromThis<FGMECanvasItemViewModel>
+	, public FEditorUndoClient
+	, public IGMETreeNodeViewModel
+{
+public:
+	/**  */
+	static TSharedRef<FGMECanvasItemViewModel> Create(const TWeakObjectPtr<const UGeometryMaskCanvas>& InCanvas);
+	virtual ~FGMECanvasItemViewModel() override = default;
+	
+	// ~Begin IGMETreeNodeViewModel
+	virtual bool GetChildren(TArray<TSharedPtr<IGMETreeNodeViewModel>>& OutChildren) override;
+	// ~End IGMETreeNodeViewModel
+
+	const FName& GetCanvasName() const { return CanvasName; }
+	const EGeometryMaskColorChannel GetColorChannel() const { return ColorChannel; }
+	const UTexture* GetCanvasTexture() const;
+	float GetMemoryUsage();
+
+private:
+	// Private token only allows members or friends to call MakeShared
+	struct FPrivateToken { explicit FPrivateToken() = default; };
+
+public:
+	FGMECanvasItemViewModel(FPrivateToken, const TWeakObjectPtr<const UGeometryMaskCanvas>& InCanvas);
+
+private:
+	FName CanvasName;
+	EGeometryMaskColorChannel ColorChannel;
+	TWeakObjectPtr<UTexture> CanvasTexture;
+	int32 KnownReaderCount;
+	int32 KnownWriterCount;
+};
diff --git a/Engine/Plugins/Experimental/GeometryMask/Source/GeometryMaskEditor/Private/ViewModels/GMECanvasListViewModel.cpp b/Engine/Plugins/Experimental/GeometryMask/Source/GeometryMaskEditor/Private/ViewModels/GMECanvasListViewModel.cpp
new file mode 100644
--- /dev/null
+++ b/Engine/Plugins/Experimental/GeometryMask/Source/GeometryMaskEditor/Private/ViewModels/GMECanvasListViewModel.cpp
@@ -0,0 +1,107 @@
+﻿// Copyright Epic Games, Inc. All Rights Reserved.
+
+#include "GMECanvasListViewModel.h"
+
+#include "Editor.h"
+#include "GeometryMaskSubsystem.h"
+#include "GMECanvasItemViewModel.h"
+
+TSharedRef<FGMECanvasListViewModel> FGMECanvasListViewModel::Create()
+{
+	TSharedRef<FGMECanvasListViewModel> ViewModel = MakeShared<FGMECanvasListViewModel>(FPrivateToken{});
+	ViewModel->Initialize();
+
+	return ViewModel;
+}
+
+FGMECanvasListViewModel::~FGMECanvasListViewModel()
+{
+	FTSTicker::GetCoreTicker().RemoveTicker(UpdateCheckHandle);
+	
+	if (UGeometryMaskSubsystem* Subsystem = GEngine->GetEngineSubsystem<UGeometryMaskSubsystem>())
+	{
+		Subsystem->OnGeometryMaskCanvasCreated().Remove(OnCanvasCreatedHandle);
+	}
+	
+	CanvasItems.Reset();
+}
+
+FGMECanvasListViewModel::FGMECanvasListViewModel(FPrivateToken)
+{
+	UpdateCheckHandle = FTSTicker::GetCoreTicker().AddTicker(FTickerDelegate::CreateRaw(this, &FGMECanvasListViewModel::Tick), UpdateCheckInterval);
+}
+
+void FGMECanvasListViewModel::Initialize()
+{
+	if (UGeometryMaskSubsystem* Subsystem = GEngine->GetEngineSubsystem<UGeometryMaskSubsystem>())
+	{
+		RefreshCanvases();
+
+		// And listen for new ones
+		{
+			OnCanvasCreatedHandle = Subsystem->OnGeometryMaskCanvasCreated().AddRaw(this, &FGMECanvasListViewModel::OnCanvasCreated);
+		}
+
+		OnChanged().Broadcast();
+	}
+}
+
+void FGMECanvasListViewModel::RefreshCanvases()
+{
+	if (UGeometryMaskSubsystem* Subsystem = GEngine->GetEngineSubsystem<UGeometryMaskSubsystem>())
+	{
+		// Add currently registered canvases
+		{
+			TArray<FName> CanvasNames = Subsystem->GetCanvasNames();
+			CanvasItems.Reset();
+			CanvasItems.Reserve(CanvasNames.Num());
+		
+			for (const FName CanvasName : CanvasNames)
+			{
+				UGeometryMaskCanvas* Canvas = Subsystem->GetNamedCanvas(CanvasName);
+				CanvasItems.Add(FGMECanvasItemViewModel::Create(Canvas));
+			}
+		}
+	}
+}
+
+void FGMECanvasListViewModel::OnCanvasCreated(const UGeometryMaskCanvas* InGeometryMaskCanvas)
+{
+	// Don't add if already in list
+	if (CanvasItems.ContainsByPredicate([CanvasName = InGeometryMaskCanvas->GetCanvasName()](const TSharedPtr<FGMECanvasItemViewModel>& InItemViewModel)
+	{
+		return InItemViewModel->GetCanvasName() == CanvasName;
+	}))
+	{
+		return;
+	}
+
+	CanvasItems.Add(FGMECanvasItemViewModel::Create(InGeometryMaskCanvas));
+
+	OnChanged().Broadcast();
+}
+
+bool FGMECanvasListViewModel::Tick(const float InDeltaSeconds)
+{
+	if (const UGeometryMaskSubsystem* Subsystem = GEngine->GetEngineSubsystem<UGeometryMaskSubsystem>())
+	{
+		const TArray<FName> CanvasNames = Subsystem->GetCanvasNames();
+		if (LastCanvasNames != CanvasNames)
+		{
+			// Canvas Names changed
+			LastCanvasNames = CanvasNames;
+			RefreshCanvases();
+			OnChanged().Broadcast();
+		}
+	}
+
+	// Always loop
+	return true;
+}
+
+bool FGMECanvasListViewModel::GetChildren(TArray<TSharedPtr<IGMETreeNodeViewModel>>& OutChildren)
+{
+	const TArray<TSharedPtr<FGMECanvasItemViewModel>> Children = CanvasItems;
+	OutChildren.Append(Children);
+	return Children.Num() > 0;
+}
diff --git a/Engine/Plugins/Experimental/GeometryMask/Source/GeometryMaskEditor/Private/ViewModels/GMECanvasListViewModel.h b/Engine/Plugins/Experimental/GeometryMask/Source/GeometryMaskEditor/Private/ViewModels/GMECanvasListViewModel.h
new file mode 100644
--- /dev/null
+++ b/Engine/Plugins/Experimental/GeometryMask/Source/GeometryMaskEditor/Private/ViewModels/GMECanvasListViewModel.h
@@ -0,0 +1,61 @@
+﻿// Copyright Epic Games, Inc. All Rights Reserved.
+
+#pragma once
+
+#include "Containers/Ticker.h"
+#include "EditorUndoClient.h"
+#include "GeometryMaskCanvas.h"
+#include "GMEViewModelShared.h"
+
+class FGMECanvasItemViewModel;
+
+class FGMECanvasListViewModel
+	: public TSharedFromThis<FGMECanvasListViewModel>
+	, public FEditorUndoClient
+	, public IGMETreeNodeViewModel
+{
+public:
+	/**  */
+	static TSharedRef<FGMECanvasListViewModel> Create();
+	virtual ~FGMECanvasListViewModel() override;
+
+	// ~Begin IGMETreeNodeViewModel
+	virtual bool GetChildren(TArray<TSharedPtr<IGMETreeNodeViewModel>>& OutChildren) override;
+	// ~End IGMETreeNodeViewModel
+
+public:
+	using FOnChanged = TMulticastDelegate<void()>;;
+
+	/** Something has changed within the ViewModel */
+	FOnChanged& OnChanged() { return OnChangedDelegate; }
+
+private:
+	// Private token only allows members or friends to call MakeShared
+	struct FPrivateToken { explicit FPrivateToken() = default; };
+	
+public:
+	explicit FGMECanvasListViewModel(FPrivateToken);
+
+private:
+	void Initialize();
+
+	void RefreshCanvases();
+
+	void OnCanvasCreated(const UGeometryMaskCanvas* InGeometryMaskCanvas);
+
+	bool Tick(const float InDeltaSeconds);
+
+private:
+	/** In seconds. */
+	static constexpr float UpdateCheckInterval = 1.0f;
+	FTSTicker::FDelegateHandle UpdateCheckHandle;
+	
+	FOnChanged OnChangedDelegate;
+	FDelegateHandle OnCanvasCreatedHandle;
+
+	/** Cached canvas names for comparison/refresh. */
+	TArray<FName> LastCanvasNames;
+
+	/** Canvas ViewModels */
+	TArray<TSharedPtr<FGMECanvasItemViewModel>> CanvasItems;
+};
diff --git a/Engine/Plugins/Experimental/GeometryMask/Source/GeometryMaskEditor/Private/ViewModels/GMEResourceItemViewModel.cpp b/Engine/Plugins/Experimental/GeometryMask/Source/GeometryMaskEditor/Private/ViewModels/GMEResourceItemViewModel.cpp
new file mode 100644
--- /dev/null
+++ b/Engine/Plugins/Experimental/GeometryMask/Source/GeometryMaskEditor/Private/ViewModels/GMEResourceItemViewModel.cpp
@@ -0,0 +1,50 @@
+﻿// Copyright Epic Games, Inc. All Rights Reserved.
+
+#include "GMEResourceItemViewModel.h"
+
+#include "Engine/TextureRenderTarget2D.h"
+#include "GeometryMaskCanvasResource.h"
+
+TSharedRef<FGMEResourceItemViewModel> FGMEResourceItemViewModel::Create(
+	const TWeakObjectPtr<const UGeometryMaskCanvasResource>& InResource)
+{
+	TSharedRef<FGMEResourceItemViewModel> ViewModel = MakeShared<FGMEResourceItemViewModel>(FPrivateToken{}, InResource);
+
+	return ViewModel;
+}
+
+float FGMEResourceItemViewModel::GetMemoryUsage() const
+{
+	if (const UTextureRenderTarget2D* Texture = GetResourceTexture())
+	{
+		return (static_cast<float>(Texture->CalcTextureMemorySizeEnum(TMC_ResidentMips)) / (1024.0f * 1024.0f));
+	}
+	
+	return 0.0f;
+}
+
+FIntPoint FGMEResourceItemViewModel::GetDimensions() const
+{
+	if (const UTextureRenderTarget2D* Texture = GetResourceTexture())
+	{
+		return {Texture->SizeX, Texture->SizeY};
+	}
+
+	return {0, 0};
+}
+
+FGMEResourceItemViewModel::FGMEResourceItemViewModel(
+	FPrivateToken,
+	const TWeakObjectPtr<const UGeometryMaskCanvasResource>& InResource)
+{
+	if (const UGeometryMaskCanvasResource* Resource = InResource.Get())
+	{
+		UniqueId = Resource->GetUniqueID();
+		ResourceTexture = const_cast<UGeometryMaskCanvasResource*>(Resource)->GetRenderTargetTexture();
+	}
+}
+
+bool FGMEResourceItemViewModel::GetChildren(TArray<TSharedPtr<IGMETreeNodeViewModel>>& OutChildren)
+{
+	return false;
+}
diff --git a/Engine/Plugins/Experimental/GeometryMask/Source/GeometryMaskEditor/Private/ViewModels/GMEResourceItemViewModel.h b/Engine/Plugins/Experimental/GeometryMask/Source/GeometryMaskEditor/Private/ViewModels/GMEResourceItemViewModel.h
new file mode 100644
--- /dev/null
+++ b/Engine/Plugins/Experimental/GeometryMask/Source/GeometryMaskEditor/Private/ViewModels/GMEResourceItemViewModel.h
@@ -0,0 +1,47 @@
+﻿// Copyright Epic Games, Inc. All Rights Reserved.
+
+#pragma once
+
+#include "Containers/Array.h"
+#include "EditorUndoClient.h"
+#include "GMEViewModelShared.h"
+#include "Templates/SharedPointer.h"
+#include "UObject/NameTypes.h"
+#include "UObject/WeakObjectPtr.h"
+
+class UGeometryMaskCanvasResource;
+class UGeometryMaskCanvas;
+class UTexture;
+class UTextureRenderTarget2D;
+
+class FGMEResourceItemViewModel
+	: public TSharedFromThis<FGMEResourceItemViewModel>
+	, public FEditorUndoClient
+	, public IGMETreeNodeViewModel
+{
+public:
+	/**  */
+	static TSharedRef<FGMEResourceItemViewModel> Create(const TWeakObjectPtr<const UGeometryMaskCanvasResource>& InResource);
+	virtual ~FGMEResourceItemViewModel() override = default;
+	
+	// ~Begin IGMETreeNodeViewModel
+	virtual bool GetChildren(TArray<TSharedPtr<IGMETreeNodeViewModel>>& OutChildren) override;
+	// ~End IGMETreeNodeViewModel
+
+	uint32 GetId() const { return UniqueId; }
+
+	const UTextureRenderTarget2D* GetResourceTexture() const { return ResourceTexture.Get(); }
+	float GetMemoryUsage() const;
+	FIntPoint GetDimensions() const;
+
+private:
+	// Private token only allows members or friends to call MakeShared
+	struct FPrivateToken { explicit FPrivateToken() = default; };
+
+public:
+	FGMEResourceItemViewModel(FPrivateToken, const TWeakObjectPtr<const UGeometryMaskCanvasResource>& InResource);
+
+private:
+	uint32 UniqueId;
+	TWeakObjectPtr<UTextureRenderTarget2D> ResourceTexture;
+};
diff --git a/Engine/Plugins/Experimental/GeometryMask/Source/GeometryMaskEditor/Private/ViewModels/GMEResourceListViewModel.cpp b/Engine/Plugins/Experimental/GeometryMask/Source/GeometryMaskEditor/Private/ViewModels/GMEResourceListViewModel.cpp
new file mode 100644
--- /dev/null
+++ b/Engine/Plugins/Experimental/GeometryMask/Source/GeometryMaskEditor/Private/ViewModels/GMEResourceListViewModel.cpp
@@ -0,0 +1,69 @@
+﻿// Copyright Epic Games, Inc. All Rights Reserved.
+
+#include "GMEResourceListViewModel.h"
+
+#include "Engine/Engine.h"
+#include "GeometryMaskSubsystem.h"
+#include "GMEResourceItemViewModel.h"
+
+TSharedRef<FGMEResourceListViewModel> FGMEResourceListViewModel::Create()
+{
+	TSharedRef<FGMEResourceListViewModel> ViewModel = MakeShared<FGMEResourceListViewModel>(FPrivateToken{});
+	ViewModel->Initialize();
+
+	return ViewModel;
+}
+
+FGMEResourceListViewModel::~FGMEResourceListViewModel()
+{
+	if (UGeometryMaskSubsystem* Subsystem = GEngine->GetEngineSubsystem<UGeometryMaskSubsystem>())
+	{
+		Subsystem->OnGeometryMaskCanvasCreated().Remove(OnResourceCreatedHandle);
+	}
+	
+	ResourceItems.Reset();
+}
+
+void FGMEResourceListViewModel::Initialize()
+{
+	if (UGeometryMaskSubsystem* Subsystem = GEngine->GetEngineSubsystem<UGeometryMaskSubsystem>())
+	{
+		// Add currently utilized canvas resources
+		ResourceItems.Reserve(4);
+		
+		for (const TObjectPtr<UGeometryMaskCanvasResource>& Resource : Subsystem->GetCanvasResources())
+		{
+			ResourceItems.Add(FGMEResourceItemViewModel::Create(Resource));
+		}
+
+		// And listen for new ones
+		{
+			OnResourceCreatedHandle = Subsystem->OnGeometryMaskResourceCreated().AddRaw(this, &FGMEResourceListViewModel::OnResourceCreated);
+		}
+
+		OnChanged().Broadcast();
+	}
+}
+
+void FGMEResourceListViewModel::OnResourceCreated(const UGeometryMaskCanvasResource* InGeometryMaskResource)
+{
+	// Don't add if already in list
+	if (ResourceItems.ContainsByPredicate([Id = InGeometryMaskResource->GetUniqueID()](const TSharedPtr<FGMEResourceItemViewModel>& InItemViewModel)
+	{
+		return InItemViewModel->GetId() == Id;
+	}))
+	{
+		return;
+	}
+
+	ResourceItems.Add(FGMEResourceItemViewModel::Create(InGeometryMaskResource));
+
+	OnChanged().Broadcast();
+}
+
+bool FGMEResourceListViewModel::GetChildren(TArray<TSharedPtr<IGMETreeNodeViewModel>>& OutChildren)
+{
+	const TArray<TSharedPtr<FGMEResourceItemViewModel>> Children = ResourceItems;
+	OutChildren.Append(Children);
+	return Children.Num() > 0;
+}
diff --git a/Engine/Plugins/Experimental/GeometryMask/Source/GeometryMaskEditor/Private/ViewModels/GMEResourceListViewModel.h b/Engine/Plugins/Experimental/GeometryMask/Source/GeometryMaskEditor/Private/ViewModels/GMEResourceListViewModel.h
new file mode 100644
--- /dev/null
+++ b/Engine/Plugins/Experimental/GeometryMask/Source/GeometryMaskEditor/Private/ViewModels/GMEResourceListViewModel.h
@@ -0,0 +1,49 @@
+﻿// Copyright Epic Games, Inc. All Rights Reserved.
+
+#pragma once
+
+#include "EditorUndoClient.h"
+#include "GeometryMaskCanvas.h"
+#include "GMEViewModelShared.h"
+
+class FGMEResourceItemViewModel;
+
+class FGMEResourceListViewModel
+	: public TSharedFromThis<FGMEResourceListViewModel>
+	, public FEditorUndoClient
+	, public IGMETreeNodeViewModel
+{
+public:
+	/**  */
+	static TSharedRef<FGMEResourceListViewModel> Create();
+	virtual ~FGMEResourceListViewModel() override;
+
+	// ~Begin IGMETreeNodeViewModel
+	virtual bool GetChildren(TArray<TSharedPtr<IGMETreeNodeViewModel>>& OutChildren) override;
+	// ~End IGMETreeNodeViewModel
+
+	public:
+	using FOnChanged = TMulticastDelegate<void()>;;
+
+	/** Something has changed within the ViewModel */
+	FOnChanged& OnChanged() { return OnChangedDelegate; }
+
+private:
+	// Private token only allows members or friends to call MakeShared
+	struct FPrivateToken { explicit FPrivateToken() = default; };
+	
+public:
+	explicit FGMEResourceListViewModel(FPrivateToken) { }
+
+protected:
+	void Initialize();
+
+	void OnResourceCreated(const UGeometryMaskCanvasResource* InGeometryMaskResource);
+
+private:
+	FOnChanged OnChangedDelegate;
+	FDelegateHandle OnResourceCreatedHandle;
+
+	/** Canvas ViewModels */
+	TArray<TSharedPtr<FGMEResourceItemViewModel>> ResourceItems;
+};
diff --git a/Engine/Plugins/Experimental/GeometryMask/Source/GeometryMaskEditor/Private/ViewModels/GMEViewModelShared.h b/Engine/Plugins/Experimental/GeometryMask/Source/GeometryMaskEditor/Private/ViewModels/GMEViewModelShared.h
new file mode 100644
--- /dev/null
+++ b/Engine/Plugins/Experimental/GeometryMask/Source/GeometryMaskEditor/Private/ViewModels/GMEViewModelShared.h
@@ -0,0 +1,17 @@
+﻿// Copyright Epic Games, Inc. All Rights Reserved.
+
+#pragma once
+
+#include "Containers/Array.h"
+#include "Templates/SharedPointerFwd.h"
+
+/** Common interface for treeview nodes. */
+class IGMETreeNodeViewModel
+{
+protected:
+	~IGMETreeNodeViewModel() = default;
+
+public:
+	/** See STreeView::OnGetChildren, return false if no children. */
+	virtual bool GetChildren(TArray<TSharedPtr<IGMETreeNodeViewModel>>& OutChildren) = 0;
+};
diff --git a/Engine/Plugins/Experimental/GeometryMask/Source/GeometryMaskEditor/Private/Widgets/SGMECanvasItem.cpp b/Engine/Plugins/Experimental/GeometryMask/Source/GeometryMaskEditor/Private/Widgets/SGMECanvasItem.cpp
new file mode 100644
--- /dev/null
+++ b/Engine/Plugins/Experimental/GeometryMask/Source/GeometryMaskEditor/Private/Widgets/SGMECanvasItem.cpp
@@ -0,0 +1,45 @@
+﻿// Copyright Epic Games, Inc. All Rights Reserved.
+
+#include "SGMECanvasItem.h"
+
+#include "Engine/Texture.h"
+#include "SlateOptMacros.h"
+#include "ViewModels/GMECanvasItemViewModel.h"
+#include "Widgets/SGeometryMaskCanvasPreview.h"
+#include "Widgets/SOverlay.h"
+#include "Widgets/Text/STextBlock.h"
+
+BEGIN_SLATE_FUNCTION_BUILD_OPTIMIZATION
+
+void SGMECanvasItem::Construct(
+	const FArguments& InArgs,
+	const TSharedRef<STableViewBase>& InOwnerTableView,
+	const TSharedRef<FGMECanvasItemViewModel>& InViewModel)
+{
+	ViewModel = InViewModel;
+
+	ImageWidget = SNew(SGeometryMaskCanvasPreview)
+					.Invert(true)
+					.SolidBackground(true)
+					.CanvasName(ViewModel->GetCanvasName())
+					.Channel(ViewModel->GetColorChannel())
+					.Opacity(1.0f);
+
+	SGMEImageItem::Construct(
+		SGMEImageItem::FArguments()
+			.Label(FText::FromName(ViewModel->GetCanvasName()))
+		, InOwnerTableView);
+}
+
+FOptionalSize SGMECanvasItem::GetAspectRatio()
+{
+	float AspectRatio = 16.0f / 9.0f;
+	if (const UTexture* CanvasTexture = ViewModel->GetCanvasTexture())
+	{
+		AspectRatio = CanvasTexture->GetSurfaceWidth() / CanvasTexture->GetSurfaceHeight();
+	}
+
+	return AspectRatio;
+}
+
+END_SLATE_FUNCTION_BUILD_OPTIMIZATION
diff --git a/Engine/Plugins/Experimental/GeometryMask/Source/GeometryMaskEditor/Private/Widgets/SGMECanvasItem.h b/Engine/Plugins/Experimental/GeometryMask/Source/GeometryMaskEditor/Private/Widgets/SGMECanvasItem.h
new file mode 100644
--- /dev/null
+++ b/Engine/Plugins/Experimental/GeometryMask/Source/GeometryMaskEditor/Private/Widgets/SGMECanvasItem.h
@@ -0,0 +1,30 @@
+﻿// Copyright Epic Games, Inc. All Rights Reserved.
+
+#pragma once
+
+#include "SGMEImageItem.h"
+#include "Templates/SharedPointer.h"
+#include "Widgets/Views/SListView.h"
+#include "Widgets/Views/STableRow.h"
+
+class FGMECanvasItemViewModel;
+class IGMETreeNodeViewModel;
+class ITableRow;
+class STableViewBase;
+
+class SGMECanvasItem
+	: public SGMEImageItem<FGMECanvasItemViewModel>
+{
+public:
+	SLATE_BEGIN_ARGS(SGMECanvasItem) {}
+	SLATE_END_ARGS()
+
+	/** Constructs this widget with InArgs */
+	void Construct(const FArguments& InArgs, const TSharedRef<STableViewBase>& InOwnerTableView, const TSharedRef<FGMECanvasItemViewModel>& InViewModel);
+
+protected:
+	virtual FOptionalSize GetAspectRatio() override;
+
+private:
+	TSharedPtr<FGMECanvasItemViewModel> ViewModel;
+};
diff --git a/Engine/Plugins/Experimental/GeometryMask/Source/GeometryMaskEditor/Private/Widgets/SGMECanvasList.cpp b/Engine/Plugins/Experimental/GeometryMask/Source/GeometryMaskEditor/Private/Widgets/SGMECanvasList.cpp
new file mode 100644
--- /dev/null
+++ b/Engine/Plugins/Experimental/GeometryMask/Source/GeometryMaskEditor/Private/Widgets/SGMECanvasList.cpp
@@ -0,0 +1,106 @@
+﻿// Copyright Epic Games, Inc. All Rights Reserved.
+
+#include "SGMECanvasList.h"
+
+#include "SGMECanvasItem.h"
+#include "SlateOptMacros.h"
+#include "ViewModels/GMECanvasItemViewModel.h"
+#include "ViewModels/GMECanvasListViewModel.h"
+#include "Widgets/Layout/SSeparator.h"
+#include "Widgets/Views/SListView.h"
+
+BEGIN_SLATE_FUNCTION_BUILD_OPTIMIZATION
+
+#define LOCTEXT_NAMESPACE "SGMECanvasList"
+
+SGMECanvasList::~SGMECanvasList()
+{
+	if (ViewModel.IsValid())
+	{
+		ViewModel->OnChanged().RemoveAll(this);
+	}
+}
+
+void SGMECanvasList::Construct(const FArguments& InArgs, TSharedRef<FGMECanvasListViewModel> InViewModel)
+{
+	constexpr float Padding = 2.0f;
+	
+	ViewModel = InViewModel;
+	ViewModel->OnChanged().AddLambda(
+		[this]()
+		{
+			Refresh();
+		});
+
+	SAssignNew(CanvasListWidget, SListView<TSharedPtr<IGMETreeNodeViewModel>>)
+		.ListItemsSource(&CanvasItems)
+		.OnGenerateRow(this, &SGMECanvasList::OnGenerateRow)
+		.ScrollbarVisibility(EVisibility::Visible)
+		.SelectionMode(ESelectionMode::None);
+
+	Refresh();
+
+	TWeakPtr<SGMECanvasList> WeakThisPtr = SharedThis(this);
+
+	ChildSlot
+	[
+		SNew(SVerticalBox)
+		+ SVerticalBox::Slot()
+		.AutoHeight()
+		[
+			SNew(SBorder)
+			.BorderImage(FAppStyle::Get().GetBrush("EditorViewportToolBar.Background"))
+			[
+				SNew(SHorizontalBox)
+				+ SHorizontalBox::Slot()
+				.AutoWidth()
+				.Padding(Padding)
+				[
+					SNew(STextBlock)
+					.Text(LOCTEXT("TotalCount", "Total:"))
+				]
+				+ SHorizontalBox::Slot()
+				.AutoWidth()
+				.Padding(Padding)
+				[
+					SNew(STextBlock)
+					.Text_Lambda([WeakThisPtr]()
+					{
+						if (WeakThisPtr.IsValid())
+						{
+							return FText::AsNumber(WeakThisPtr.Pin()->CanvasItems.Num());
+						}
+						return FText::AsNumber(0);
+					})
+				]
+			]
+		]
+		+ SVerticalBox::Slot()
+		.FillHeight(1.0f)
+		.Padding(1, 1, 1, Padding)
+		[
+			CanvasListWidget.ToSharedRef()
+		]
+	];
+}
+
+TSharedRef<ITableRow> SGMECanvasList::OnGenerateRow(
+	TSharedPtr<IGMETreeNodeViewModel> InViewModel
+	,const TSharedRef<STableViewBase>& InOwnerTable)
+{
+	check(InViewModel.IsValid());
+
+	return SNew(SGMECanvasItem, InOwnerTable, StaticCastSharedPtr<FGMECanvasItemViewModel>(InViewModel).ToSharedRef());
+}
+
+void SGMECanvasList::Refresh()
+{
+	CanvasItems.Reset();
+	ViewModel->GetChildren(CanvasItems);
+	
+	CanvasListWidget->RequestListRefresh();
+}
+
+#undef LOCTEXT_NAMESPACE
+
+END_SLATE_FUNCTION_BUILD_OPTIMIZATION
diff --git a/Engine/Plugins/Experimental/GeometryMask/Source/GeometryMaskEditor/Private/Widgets/SGMECanvasList.h b/Engine/Plugins/Experimental/GeometryMask/Source/GeometryMaskEditor/Private/Widgets/SGMECanvasList.h
new file mode 100644
--- /dev/null
+++ b/Engine/Plugins/Experimental/GeometryMask/Source/GeometryMaskEditor/Private/Widgets/SGMECanvasList.h
@@ -0,0 +1,37 @@
+﻿// Copyright Epic Games, Inc. All Rights Reserved.
+
+#pragma once
+
+#include "Widgets/SCompoundWidget.h"
+
+class FGMECanvasListViewModel;
+class IGMETreeNodeViewModel;
+class ITableRow;
+class STableViewBase;
+template <typename ItemType> class SListView;
+
+class SGMECanvasList
+	: public SCompoundWidget
+{
+public:
+	SLATE_BEGIN_ARGS(SGMECanvasList)
+	{ }
+	SLATE_END_ARGS()
+
+	virtual ~SGMECanvasList() override;
+
+	/** Constructs this widget with InArgs */
+	void Construct(const FArguments& InArgs, TSharedRef<FGMECanvasListViewModel> InViewModel);
+
+private:
+	TSharedRef<ITableRow> OnGenerateRow(TSharedPtr<IGMETreeNodeViewModel> InViewModel, const TSharedRef<STableViewBase>& InOwnerTable);
+
+	void Refresh();
+	
+private:
+	TSharedPtr<FGMECanvasListViewModel> ViewModel;
+	
+	TSharedPtr<SListView<TSharedPtr<IGMETreeNodeViewModel>>> CanvasListWidget;
+
+	TArray<TSharedPtr<IGMETreeNodeViewModel>> CanvasItems;	
+};
diff --git a/Engine/Plugins/Experimental/GeometryMask/Source/GeometryMaskEditor/Private/Widgets/SGMEImageItem.h b/Engine/Plugins/Experimental/GeometryMask/Source/GeometryMaskEditor/Private/Widgets/SGMEImageItem.h
new file mode 100644
--- /dev/null
+++ b/Engine/Plugins/Experimental/GeometryMask/Source/GeometryMaskEditor/Private/Widgets/SGMEImageItem.h
@@ -0,0 +1,94 @@
+// Copyright Epic Games, Inc. All Rights Reserved.
+
+#pragma once
+
+#include "Templates/SharedPointer.h"
+#include "Widgets/Colors/SColorBlock.h"
+#include "Widgets/Layout/SScaleBox.h"
+#include "Widgets/Layout/SWrapBox.h"
+#include "Widgets/Views/SListView.h"
+#include "Widgets/Views/STableRow.h"
+
+class FGMECanvasItemViewModel;
+class IGMETreeNodeViewModel;
+class ITableRow;
+class STableViewBase;
+
+template <typename InViewModelType>
+class SGMEImageItem;
+
+template <typename InViewModelType>
+class SGMEImageItem : public STableRow<TSharedPtr<InViewModelType>>
+{
+	using FTableRowArgs = typename STableRow<TSharedPtr<InViewModelType>>::FArguments;
+
+public:
+	SLATE_BEGIN_ARGS(SGMEImageItem)
+		{ }
+		SLATE_ATTRIBUTE(FText, Label)
+	SLATE_END_ARGS()
+
+	/** Constructs this widget with InArgs */
+	void Construct(const FArguments& InArgs, const TSharedRef<STableViewBase>& InOwnerTableView)
+	{
+		constexpr float Padding = 2.0f;
+
+		TAttribute<FOptionalSize> AspectRatioBinding;
+		AspectRatioBinding.Bind(TAttribute<FOptionalSize>::FGetter::CreateSP(this, &SGMEImageItem::GetAspectRatio));
+
+		TAttribute<FOptionalSize> HeightBinding;
+		HeightBinding.Bind(TAttribute<FOptionalSize>::FGetter::CreateSP(this, &SGMEImageItem::GetHeight));
+
+		FTableRowArgs Args = FTableRowArgs()
+			.Padding(Padding)
+			.ShowSelection(false)
+			[
+				SNew(SVerticalBox)
+				+ SVerticalBox::Slot()
+				.VAlign(VAlign_Top)
+				[
+					SAssignNew(ImageContainerWidget, SBox)
+					.HeightOverride(HeightBinding)
+					.MinAspectRatio(AspectRatioBinding)
+					.MaxAspectRatio(AspectRatioBinding)
+					[
+						ImageWidget.ToSharedRef()
+					]
+				]
+
+				+ SVerticalBox::Slot()
+				.Padding(1, 1, 1, Padding)
+				[
+					SNew(SHorizontalBox)
+					+ SHorizontalBox::Slot()
+					[
+						SNew(STextBlock)
+						.Text(InArgs._Label)
+					]
+				]
+			];
+
+		STableRow<TSharedPtr<InViewModelType>>::Construct(Args, InOwnerTableView);
+	}
+
+protected:
+	virtual FOptionalSize GetAspectRatio() = 0;
+
+	/** Get desired height based on the actual width. */
+	FOptionalSize GetHeight()
+	{
+		if (ImageContainerWidget.IsValid())
+		{
+			const FGeometry& WidgetGeometry = ImageContainerWidget->GetCachedGeometry();
+			const float AspectRatio = 1.0f / GetAspectRatio().Get();
+			const float Height = static_cast<int32>(WidgetGeometry.GetLocalSize().X) * AspectRatio;
+			return Height * 2; // Don't know why it's x2, it just is.
+		}
+
+		return {};
+	}
+
+protected:
+	TSharedPtr<SPanel> ImageContainerWidget;
+	TSharedPtr<SWidget> ImageWidget;
+};
diff --git a/Engine/Plugins/Experimental/GeometryMask/Source/GeometryMaskEditor/Private/Widgets/SGMEResourceItem.cpp b/Engine/Plugins/Experimental/GeometryMask/Source/GeometryMaskEditor/Private/Widgets/SGMEResourceItem.cpp
new file mode 100644
--- /dev/null
+++ b/Engine/Plugins/Experimental/GeometryMask/Source/GeometryMaskEditor/Private/Widgets/SGMEResourceItem.cpp
@@ -0,0 +1,47 @@
+﻿// Copyright Epic Games, Inc. All Rights Reserved.
+
+#include "SGMEResourceItem.h"
+
+#include "Engine/TextureRenderTarget2D.h"
+#include "SlateOptMacros.h"
+#include "ViewModels/GMEResourceItemViewModel.h"
+
+BEGIN_SLATE_FUNCTION_BUILD_OPTIMIZATION
+
+void SGMEResourceItem::Construct(
+	const FArguments& InArgs,
+	const TSharedRef<STableViewBase>& InOwnerTableView,
+	const TSharedRef<FGMEResourceItemViewModel>& InViewModel)
+{
+	ViewModel = InViewModel;
+
+	TextureBrush = MakeShared<FSlateBrush>();
+	TextureBrush->SetResourceObject(const_cast<UTextureRenderTarget2D*>(InViewModel->GetResourceTexture()));
+	TextureBrush->ImageSize = FVector2D(InViewModel->GetResourceTexture()->SizeX, InViewModel->GetResourceTexture()->SizeY);
+
+	ImageWidget = SNew(SImage)
+					.Image(TextureBrush.IsValid()
+						? TextureBrush.Get()
+						: FAppStyle::GetBrush("WhiteTexture"));
+	
+	SGMEImageItem::Construct(
+		SGMEImageItem::FArguments()
+			.Label_Lambda([InViewModel]()
+			{
+				return FText::FromString(InViewModel->GetDimensions().ToString());
+			})
+		, InOwnerTableView);
+}
+
+FOptionalSize SGMEResourceItem::GetAspectRatio()
+{
+	float AspectRatio = 16.0f / 9.0f;
+	if (const UTexture* ResourceTexture = ViewModel->GetResourceTexture())
+	{
+		AspectRatio = ResourceTexture->GetSurfaceWidth() / ResourceTexture->GetSurfaceHeight();
+	}
+
+	return AspectRatio;
+}
+
+END_SLATE_FUNCTION_BUILD_OPTIMIZATION
diff --git a/Engine/Plugins/Experimental/GeometryMask/Source/GeometryMaskEditor/Private/Widgets/SGMEResourceItem.h b/Engine/Plugins/Experimental/GeometryMask/Source/GeometryMaskEditor/Private/Widgets/SGMEResourceItem.h
new file mode 100644
--- /dev/null
+++ b/Engine/Plugins/Experimental/GeometryMask/Source/GeometryMaskEditor/Private/Widgets/SGMEResourceItem.h
@@ -0,0 +1,33 @@
+﻿// Copyright Epic Games, Inc. All Rights Reserved.
+
+#pragma once
+
+#include "SGMEImageItem.h"
+#include "Templates/SharedPointer.h"
+#include "Widgets/Views/SListView.h"
+#include "Widgets/Views/STableRow.h"
+
+class FGMEResourceItemViewModel;
+class IGMETreeNodeViewModel;
+class ITableRow;
+class STableViewBase;
+
+class SGMEResourceItem
+	: public SGMEImageItem<FGMEResourceItemViewModel>
+{
+public:
+	SLATE_BEGIN_ARGS(SGMEResourceItem) {}
+	SLATE_END_ARGS()
+
+	/** Constructs this widget with InArgs */
+	void Construct(const FArguments& InArgs, const TSharedRef<STableViewBase>& InOwnerTableView, const TSharedRef<FGMEResourceItemViewModel>& InViewModel);
+
+protected:
+	virtual FOptionalSize GetAspectRatio() override;
+	
+private:
+	TSharedPtr<FGMEResourceItemViewModel> ViewModel;
+
+	/** The Slate brush that renders the texture. */
+	TSharedPtr<FSlateBrush> TextureBrush;
+};
diff --git a/Engine/Plugins/Experimental/GeometryMask/Source/GeometryMaskEditor/Private/Widgets/SGMEResourceList.cpp b/Engine/Plugins/Experimental/GeometryMask/Source/GeometryMaskEditor/Private/Widgets/SGMEResourceList.cpp
new file mode 100644
--- /dev/null
+++ b/Engine/Plugins/Experimental/GeometryMask/Source/GeometryMaskEditor/Private/Widgets/SGMEResourceList.cpp
@@ -0,0 +1,142 @@
+﻿// Copyright Epic Games, Inc. All Rights Reserved.
+
+#include "SGMEResourceList.h"
+
+#include "SGMEResourceItem.h"
+#include "SlateOptMacros.h"
+#include "ViewModels/GMEResourceItemViewModel.h"
+#include "ViewModels/GMEResourceListViewModel.h"
+#include "Widgets/Layout/SSeparator.h"
+
+BEGIN_SLATE_FUNCTION_BUILD_OPTIMIZATION
+
+#define LOCTEXT_NAMESPACE "SGMEResourceList"
+
+SGMEResourceList::~SGMEResourceList()
+{
+	if (ViewModel.IsValid())
+	{
+		ViewModel->OnChanged().RemoveAll(this);
+	}
+}
+
+void SGMEResourceList::Construct(const FArguments& InArgs, TSharedRef<FGMEResourceListViewModel> InViewModel)
+{
+	constexpr float Padding = 2.0f;
+	
+	ViewModel = InViewModel;
+	ViewModel->OnChanged().AddSPLambda(
+		this,
+		[this]()
+		{
+			Refresh();
+		});
+
+	SAssignNew(ResourceListWidget, SListView<TSharedPtr<IGMETreeNodeViewModel>>)
+		.ListItemsSource(&ResourceItems)
+		.OnGenerateRow(this, &SGMEResourceList::OnGenerateRow)
+		.ScrollbarVisibility(EVisibility::Visible)
+		.SelectionMode(ESelectionMode::None);
+
+	Refresh();
+
+	TWeakPtr<SGMEResourceList> WeakThisPtr = SharedThis(this);
+
+	ChildSlot
+	[
+		SNew(SVerticalBox)
+		+ SVerticalBox::Slot()
+		.AutoHeight()
+		[
+			SNew(SBorder)
+			.BorderImage(FAppStyle::Get().GetBrush("EditorViewportToolBar.Background"))
+			[
+				SNew(SHorizontalBox)
+				+ SHorizontalBox::Slot()
+				.AutoWidth()
+				.Padding(Padding)
+				[
+					SNew(STextBlock)
+					.Text(LOCTEXT("TotalCount", "Total:"))
+				]
+				+ SHorizontalBox::Slot()
+				.AutoWidth()
+				.Padding(Padding)
+				[
+					SNew(STextBlock)
+					.Text_Lambda([WeakThisPtr]()
+					{
+						if (WeakThisPtr.IsValid())
+						{
+							return FText::AsNumber(WeakThisPtr.Pin()->ResourceItems.Num());
+						}
+						return FText::AsNumber(0);
+					})
+				]
+				+ SHorizontalBox::Slot()
+				.AutoWidth()
+				.Padding(Padding)
+				[
+					SNew(SSeparator)
+					.Orientation(Orient_Vertical)
+				]
+				+ SHorizontalBox::Slot()
+				.AutoWidth()
+				.Padding(Padding)
+				[
+					SNew(STextBlock)
+					.Text(LOCTEXT("TotalMemory", "Memory:"))
+				]
+				+ SHorizontalBox::Slot()
+				.AutoWidth()
+				.Padding(Padding)
+				[
+					SNew(STextBlock)
+					.Text_Lambda([WeakThisPtr]()
+					{
+						if (WeakThisPtr.IsValid())
+						{
+							float Memory = 0.0f;
+							for (TSharedPtr<IGMETreeNodeViewModel>& ChildItem : WeakThisPtr.Pin()->ResourceItems)
+							{
+								if (TSharedPtr<FGMEResourceItemViewModel> ResourceItem = StaticCastSharedPtr<FGMEResourceItemViewModel>(ChildItem))
+								{
+									Memory += ResourceItem->GetMemoryUsage();
+								}								
+							}
+							return FText::Format(LOCTEXT("TotalMemoryValue", "{0}mb"), FText::AsNumber(Memory));
+						}
+						return FText::AsNumber(0);
+					})
+				]
+			]
+		]
+		+ SVerticalBox::Slot()
+		.FillHeight(1.0f)
+		.Padding(1, 1, 1, Padding)
+		[
+			ResourceListWidget.ToSharedRef()
+		]
+	];
+}
+
+TSharedRef<ITableRow> SGMEResourceList::OnGenerateRow(
+	TSharedPtr<IGMETreeNodeViewModel> InViewModel
+	, const TSharedRef<STableViewBase>& InOwnerTable)
+{
+	check(InViewModel.IsValid());
+
+	return SNew(SGMEResourceItem, InOwnerTable, StaticCastSharedPtr<FGMEResourceItemViewModel>(InViewModel).ToSharedRef());
+}
+
+void SGMEResourceList::Refresh()
+{
+	ResourceItems.Reset();
+	ViewModel->GetChildren(ResourceItems);
+	
+	ResourceListWidget->RequestListRefresh();
+}
+
+#undef LOCTEXT_NAMESPACE
+
+END_SLATE_FUNCTION_BUILD_OPTIMIZATION
diff --git a/Engine/Plugins/Experimental/GeometryMask/Source/GeometryMaskEditor/Private/Widgets/SGMEResourceList.h b/Engine/Plugins/Experimental/GeometryMask/Source/GeometryMaskEditor/Private/Widgets/SGMEResourceList.h
new file mode 100644
--- /dev/null
+++ b/Engine/Plugins/Experimental/GeometryMask/Source/GeometryMaskEditor/Private/Widgets/SGMEResourceList.h
@@ -0,0 +1,38 @@
+﻿// Copyright Epic Games, Inc. All Rights Reserved.
+
+#pragma once
+
+#include "Widgets/SCompoundWidget.h"
+
+class FGMEResourceListViewModel;
+class IGMETreeNodeViewModel;
+class ITableRow;
+class STableViewBase;
+template <typename ItemType> class SListView;
+
+class SGMEResourceList
+	: public SCompoundWidget
+{
+public:
+	SLATE_BEGIN_ARGS(SGMEResourceList)
+	{ }
+	SLATE_END_ARGS()
+
+	virtual ~SGMEResourceList() override;
+
+	/** Constructs this widget with InArgs */
+	void Construct(const FArguments& InArgs, TSharedRef<FGMEResourceListViewModel> InViewModel);
+
+private:
+	TSharedRef<ITableRow> OnGenerateRow(TSharedPtr<IGMETreeNodeViewModel> InViewModel
+		, const TSharedRef<STableViewBase>& InOwnerTable);
+
+	void Refresh();
+
+private:
+	TSharedPtr<FGMEResourceListViewModel> ViewModel;
+
+	TSharedPtr<SListView<TSharedPtr<IGMETreeNodeViewModel>>> ResourceListWidget;
+
+	TArray<TSharedPtr<IGMETreeNodeViewModel>> ResourceItems;
+};
diff --git a/Engine/Plugins/Experimental/GeometryMask/Source/GeometryMaskEditor/Private/Widgets/SGeometryMaskCanvasPreview.cpp b/Engine/Plugins/Experimental/GeometryMask/Source/GeometryMaskEditor/Private/Widgets/SGeometryMaskCanvasPreview.cpp
new file mode 100644
--- /dev/null
+++ b/Engine/Plugins/Experimental/GeometryMask/Source/GeometryMaskEditor/Private/Widgets/SGeometryMaskCanvasPreview.cpp
@@ -0,0 +1,289 @@
+﻿// Copyright Epic Games, Inc. All Rights Reserved.
+
+#include "Widgets/SGeometryMaskCanvasPreview.h"
+
+#include "Engine/Engine.h"
+#include "Engine/Texture.h"
+#include "Engine/TextureRenderTarget2D.h"
+#include "GeometryMaskEditorLog.h"
+#include "GeometryMaskSubsystem.h"
+#include "GeometryMaskTypes.h"
+#include "Materials/MaterialInstanceDynamic.h"
+#include "Materials/MaterialInterface.h"
+#include "SlateMaterialBrush.h"
+#include "SlateOptMacros.h"
+#include "UObject/Package.h"
+#include "Widgets/Colors/SColorBlock.h"
+#include "Widgets/Images/SImage.h"
+#include "Widgets/Layout/SScaleBox.h"
+
+namespace UE::GeometryMaskEditor::Private
+{
+	constexpr const TCHAR* PreviewMaterialPath = TEXT("/GeometryMask/GeometryMask/M_GeometryMaskPreview");
+	constexpr const TCHAR* DefaultTexturePath = TEXT("/Engine/EngineResources/Black");
+
+	static FName TextureParameterName = TEXT("Mask_Texture");
+	static FName ChannelParameterName = TEXT("Mask_Channel");
+	static FName InvertParameterName = TEXT("Mask_Invert");
+	static FName PaddingParameterName = TEXT("Mask_Padding");
+	static FName OpacityParameterName = TEXT("OpacityMultiplier");
+	static FName FeatherParameterName = TEXT("Feather");
+}
+
+BEGIN_SLATE_FUNCTION_BUILD_OPTIMIZATION
+
+SGeometryMaskCanvasPreview::SGeometryMaskCanvasPreview()
+	: PreviewMaterialPath(UE::GeometryMaskEditor::Private::PreviewMaterialPath)
+	, DefaultTexturePath(UE::GeometryMaskEditor::Private::DefaultTexturePath)
+{
+}
+
+void SGeometryMaskCanvasPreview::Construct(const FArguments& InArgs)
+{
+	CanvasName = InArgs._CanvasName;
+	ColorChannel = InArgs._Channel;
+	bInvert = InArgs._Invert;
+	bHasSolidBackground = InArgs._SolidBackground;
+	OpacityMultiplier = InArgs._Opacity;
+	AspectRatio = 1920.0f / 1080.0f;
+
+	// Setup MID and Brush
+	{
+		if (UMaterialInterface* PreviewMaterial = Cast<UMaterialInterface>(PreviewMaterialPath.TryLoad()))
+		{
+			PreviewMID = UMaterialInstanceDynamic::Create(PreviewMaterial, GetTransientPackage());
+			PreviewMID->SetScalarParameterValue(UE::GeometryMaskEditor::Private::OpacityParameterName, OpacityMultiplier.Get(1.0f));
+			PreviewBrush = MakeShared<FSlateMaterialBrush>(*PreviewMID, FVector2D(1920.f, 1080.f));
+		}
+		else
+		{
+			UE_LOG(LogGeometryMaskEditor, Warning, TEXT("PreviewMaterial could not be loaded from path: '%s'"), *PreviewMaterialPath.ToString());
+		}
+	}
+
+	// Load DefaultTexture
+	{
+		if (DefaultTexture = Cast<UTexture>(DefaultTexturePath.TryLoad());
+			!DefaultTexture || !IsValid(DefaultTexture))
+		{
+			UE_LOG(LogGeometryMaskEditor, Warning, TEXT("DefaultTexture could not be loaded from path: '%s'"), *DefaultTexturePath.ToString());
+		}
+	}
+
+	UpdateBrush(nullptr, DefaultTexture);
+    TryResolveCanvas();
+
+	ChildSlot
+	[
+		SNew(SScaleBox)
+		.Stretch(EStretch::ScaleToFill)
+		[
+			SNew(SOverlay)
+			+ SOverlay::Slot()
+			[
+				SNew(SImage)
+				.Visibility_Lambda([this]()
+				{
+					return HasSolidBackground()
+						? EVisibility::Visible
+						: EVisibility::Hidden;
+				})
+				.Image(FAppStyle::GetBrush("WhiteTexture"))
+			]
+
+			+ SOverlay::Slot()
+			[
+				SNew(SImage)
+				.Image(PreviewBrush.IsValid() ? PreviewBrush.Get() : FAppStyle::GetBrush("WhiteTexture"))
+			]
+		]
+	];
+}
+
+void SGeometryMaskCanvasPreview::Tick(const FGeometry& InAllottedGeometry, const double InCurrentTime, const float InDeltaTime)
+{
+	SCompoundWidget::Tick(InAllottedGeometry, InCurrentTime, InDeltaTime);
+
+	UpdateBrush(Canvas.Get(), nullptr);
+}
+
+const FName SGeometryMaskCanvasPreview::GetCanvasName() const
+{
+	return CanvasName.Get(NAME_None);
+}
+
+void SGeometryMaskCanvasPreview::SetCanvasName(const FName InCanvasName)
+{
+	if (GetCanvasName() != InCanvasName)
+	{
+		CanvasName.Set(InCanvasName);
+		TryResolveCanvas();
+	}
+}
+
+const EGeometryMaskColorChannel SGeometryMaskCanvasPreview::GetColorChannel(const bool bOnlyValid) const
+{
+	EGeometryMaskColorChannel Result = ColorChannel.Get(EGeometryMaskColorChannel::Red);
+	if (bOnlyValid)
+	{
+		Result = UE::GeometryMask::GetValidMaskChannel(Result, false);
+	}
+
+	return Result;
+}
+
+void SGeometryMaskCanvasPreview::SetColorChannel(const EGeometryMaskColorChannel InColorChannel)
+{
+	if (InColorChannel == EGeometryMaskColorChannel::None)
+	{
+		return;
+	}
+
+	ColorChannel.Set(InColorChannel);
+	
+	if (PreviewMID)
+	{
+		PreviewMID->SetVectorParameterValue(UE::GeometryMaskEditor::Private::ChannelParameterName, UE::GeometryMask::MaskChannelEnumToVector[InColorChannel]);
+	}
+}
+
+const bool SGeometryMaskCanvasPreview::IsInverted() const
+{
+	return bInvert.Get(false);
+}
+
+void SGeometryMaskCanvasPreview::SetInvert(const bool bInInverted)
+{
+	bInvert.Set(bInInverted);
+}
+
+const bool SGeometryMaskCanvasPreview::HasSolidBackground() const
+{
+	return bHasSolidBackground.Get();
+}
+
+void SGeometryMaskCanvasPreview::SetSolidBackground(const bool bInHasSolidBackground)
+{
+	bHasSolidBackground.Set(bInHasSolidBackground);
+}
+
+const float SGeometryMaskCanvasPreview::GetOpacity() const
+{
+	return OpacityMultiplier.Get(1.0f);
+}
+
+void SGeometryMaskCanvasPreview::SetOpacity(const float InValue)
+{
+	OpacityMultiplier.Set(InValue);
+
+	if (PreviewMID)
+	{
+		PreviewMID->SetScalarParameterValue(UE::GeometryMaskEditor::Private::OpacityParameterName, InValue);
+	}
+}
+
+UGeometryMaskCanvas* SGeometryMaskCanvasPreview::GetCanvas() const
+{
+	return Canvas.Get();
+}
+
+FOptionalSize SGeometryMaskCanvasPreview::GetAspectRatio()
+{
+	float Width = 1920.0f;
+	float Height = 1080.0f;
+
+	if (PreviewBrush.IsValid())
+	{
+		Width = PreviewBrush->ImageSize.X;
+		Height = PreviewBrush->ImageSize.Y;		
+	}
+	
+	if (AspectRatio.IsSet())
+	{
+		return AspectRatio.Get();
+	}
+
+	AspectRatio.Set(Width / Height);
+	return AspectRatio.Get();
+}
+
+void SGeometryMaskCanvasPreview::AddReferencedObjects(FReferenceCollector& InCollector)
+{
+	InCollector.AddReferencedObject(PreviewMID);
+	InCollector.AddReferencedObject(DefaultTexture);
+}
+
+FString SGeometryMaskCanvasPreview::GetReferencerName() const
+{
+	return TEXT("SGeometryMaskCanvasPreview");
+}
+
+bool SGeometryMaskCanvasPreview::TryResolveCanvas()
+{
+	if (UGeometryMaskSubsystem* Subsystem = GEngine->GetEngineSubsystem<UGeometryMaskSubsystem>())
+	{
+		Canvas = Subsystem->GetNamedCanvas(GetCanvasName());
+		UpdateBrush(Canvas.Get(), nullptr);
+	}
+
+	return Canvas != nullptr;
+}
+
+void SGeometryMaskCanvasPreview::UpdateBrush(const UGeometryMaskCanvas* InCanvas, UTexture* InTexture)
+{
+	// If InTexture isn't specified, use Default/blank
+	if (!InTexture)
+	{
+		InTexture = DefaultTexture;
+	}
+
+	const bool bHasValidMID = PreviewMID != nullptr;
+		
+	// When a canvas is provided, update the parameters from it, vs just from this widget
+	if (InCanvas)
+	{
+		if (UTexture* CanvasTexture = InCanvas->GetTexture())
+		{
+			InTexture = CanvasTexture;
+		}
+
+		FVector4f Padding(ForceInitToZero);
+		if (const UGeometryMaskCanvasResource* CanvasResource = InCanvas->GetResource())
+		{
+			FIntPoint ViewportPadding(CanvasResource->GetViewportPadding());
+			FVector2f ViewportPaddingF = FVector2f(ViewportPadding) / CanvasResource->GetViewportSize();
+			Padding.X = ViewportPaddingF.X;
+			Padding.Y = ViewportPaddingF.Y;
+		}
+
+		if (bHasValidMID)
+		{
+			const int32 MaxFeatherRadius = FMath::Max(InCanvas->GetOuterFeatherRadius(), InCanvas->GetInnerFeatherRadius());
+			const FVector4f Feather(InCanvas->IsFeatherApplied() ? 1 : 0, InCanvas->GetOuterFeatherRadius(), InCanvas->GetInnerFeatherRadius(), MaxFeatherRadius);
+			PreviewMID->SetVectorParameterValue(UE::GeometryMaskEditor::Private::FeatherParameterName, Feather);
+			PreviewMID->SetVectorParameterValue(UE::GeometryMaskEditor::Private::PaddingParameterName, Padding);
+		}
+	}
+	
+	if (bHasValidMID)
+	{
+		PreviewMID->SetTextureParameterValue(UE::GeometryMaskEditor::Private::TextureParameterName, InTexture);
+		PreviewMID->SetVectorParameterValue(UE::GeometryMaskEditor::Private::ChannelParameterName, UE::GeometryMask::MaskChannelEnumToVector[GetColorChannel()]);
+		PreviewMID->SetScalarParameterValue(UE::GeometryMaskEditor::Private::InvertParameterName, IsInverted() ? 1.0f : 0.0f);
+	}
+
+	if (PreviewBrush.IsValid())
+	{
+		FVector2D ImageSize(1920.f, 1080.f);
+		if (InTexture)
+		{
+			ImageSize.X = InTexture->GetSurfaceWidth();
+			ImageSize.Y = InTexture->GetSurfaceHeight();
+		}
+		
+		PreviewBrush->ImageSize = ImageSize;
+		AspectRatio.Set(ImageSize.X / ImageSize.Y);
+	}
+}
+
+END_SLATE_FUNCTION_BUILD_OPTIMIZATION
diff --git a/Engine/Plugins/Experimental/GeometryMask/Source/GeometryMaskEditor/Public/Widgets/SGeometryMaskCanvasPreview.h b/Engine/Plugins/Experimental/GeometryMask/Source/GeometryMaskEditor/Public/Widgets/SGeometryMaskCanvasPreview.h
new file mode 100644
--- /dev/null
+++ b/Engine/Plugins/Experimental/GeometryMask/Source/GeometryMaskEditor/Public/Widgets/SGeometryMaskCanvasPreview.h
@@ -0,0 +1,108 @@
+﻿// Copyright Epic Games, Inc. All Rights Reserved.
+
+#pragma once
+
+#include "GeometryMaskTypes.h"
+#include "Types/SlateStructs.h"
+#include "UObject/GCObject.h"
+#include "Widgets/SCompoundWidget.h"
+
+class UGeometryMaskCanvas;
+class UMaterialInstanceDynamic;
+class UTexture;
+
+/** Displays a named GeometryMaskCanvas. */
+class GEOMETRYMASKEDITOR_API SGeometryMaskCanvasPreview
+	: public SCompoundWidget
+	, public FGCObject
+{
+public:
+	SLATE_BEGIN_ARGS(SGeometryMaskCanvasPreview) {}
+		SLATE_ATTRIBUTE(FName, CanvasName)
+		SLATE_ATTRIBUTE(EGeometryMaskColorChannel, Channel)
+		SLATE_ATTRIBUTE(bool, Invert)
+		SLATE_ATTRIBUTE(bool, SolidBackground)
+		SLATE_ATTRIBUTE(float, Opacity)
+	SLATE_END_ARGS()
+
+	SGeometryMaskCanvasPreview();
+
+	/** Constructs this widget with InArgs */
+	void Construct(const FArguments& InArgs);
+
+	virtual void Tick(const FGeometry& InAllottedGeometry, const double InCurrentTime, const float InDeltaTime) override;
+
+	/** Get the name of the currently referenced canvas. */
+	const FName GetCanvasName() const;
+
+	/** Sets the name of the currently referenced canvas, and resolves the canvas itself. */
+	void SetCanvasName(const FName InCanvasName);
+
+	/** Get the color channel to display for the canvas.  */
+	const EGeometryMaskColorChannel GetColorChannel(const bool bOnlyValid = true) const;
+
+	/** Set the color channel to display for the canvas.  */
+	void SetColorChannel(const EGeometryMaskColorChannel InColorChannel);
+
+	/** Get whether to invert the display of the canvas.  */
+	const bool IsInverted() const;
+
+	/** Set whether to invert the display of the canvas.  */
+	void SetInvert(const bool bInInverted);
+
+	/** Get whether a solid background is used or not (vs. compositing with alpha on widgets below). */
+	const bool HasSolidBackground() const;
+
+	/** Set whether a solid background is used or not (vs. compositing with alpha on widgets below).  */
+	void SetSolidBackground(const bool bInHasSolidBackground);
+
+	/** Get the overall opacity multiplier value. */
+	const float GetOpacity() const;
+
+	/** Set the overall opacity multiplier value. */
+	void SetOpacity(const float InValue);
+
+	/** Get the currently referenced canvas. */
+	UGeometryMaskCanvas* GetCanvas() const;
+
+	/** Get the aspect ratio of the referenced canvas. */
+	FOptionalSize GetAspectRatio();
+
+	// ~Begin FGCObject Interface
+	virtual void AddReferencedObjects(FReferenceCollector& InCollector) override;
+	virtual FString GetReferencerName() const override;
+	// ~End FGCObject Interface
+
+private:
+	bool TryResolveCanvas();
+
+	void UpdateBrush(const UGeometryMaskCanvas* InCanvas, UTexture* InTexture);
+
+private:
+	TSharedPtr<FSlateBrush> PreviewBrush;
+	FSoftObjectPath PreviewMaterialPath;
+	TObjectPtr<UMaterialInstanceDynamic> PreviewMID;
+
+	FSoftObjectPath DefaultTexturePath;
+	TObjectPtr<UTexture> DefaultTexture;
+
+	TWeakObjectPtr<UGeometryMaskCanvas> Canvas;
+
+	/** Name of the currently referenced canvas. */
+	TAttribute<FName> CanvasName;
+
+	/** Color channel to display for the canvas.  */
+	TAttribute<EGeometryMaskColorChannel> ColorChannel;
+
+	/** Whether to invert the display of the canvas.  */
+	TAttribute<bool> bInvert;
+
+	/*** Whether a solid background is used or not (vs. compositing with alpha on widgets below). */
+	TAttribute<bool> bHasSolidBackground;
+
+	/** Multiplies the overall opacity. */
+	TAttribute<float> OpacityMultiplier;
+
+	/** Stores the aspect ratio of the referenced canvas. */
+	TAttribute<FOptionalSize> AspectRatio;
+};
