[GeometryMask] Perf improvements:
- Feathering off by default
- GPU Buffer caching
- MSAA support (disabled when using Blur or Feather)
- More selective updates/drawing
- Added stats and console commands to assist in debugging

diff --git a/Engine/Plugins/Experimental/Avalanche/Source/AvalancheMask/Private/Handling/AvaMaskAvaShapeMaterialCollectionHandle.cpp b/Engine/Plugins/Experimental/Avalanche/Source/AvalancheMask/Private/Handling/AvaMaskAvaShapeMaterialCollectionHandle.cpp
--- a/Engine/Plugins/Experimental/Avalanche/Source/AvalancheMask/Private/Handling/AvaMaskAvaShapeMaterialCollectionHandle.cpp
+++ b/Engine/Plugins/Experimental/Avalanche/Source/AvalancheMask/Private/Handling/AvaMaskAvaShapeMaterialCollectionHandle.cpp
@@ -85,20 +85,22 @@ TArray<TSharedPtr<IAvaMaskMaterialHandle>> FAvaMaskAvaShapeMaterialCollectionHan
 		RefreshMaterials();
 	}
 	
-	for (const TPair<int32, TSharedPtr<IAvaMaskMaterialHandle>>& MaterialHandle : MaterialHandles)
+	for (TPair<int32, TSharedPtr<IAvaMaskMaterialHandle>>& SlotMaterialHandle : MaterialHandles)
 	{
-		if (IsSlotOccupied(MaterialHandle.Key) && !UE::Ava::Internal::IsHandleValid(MaterialHandle.Value))
+		// Account for empty/unoccupied slot
+		if (TSharedPtr<IAvaMaskMaterialHandle>& MaterialHandle = SlotMaterialHandle.Value)
 		{
-			if (RefreshMaterials())
+			if (!UE::Ava::Internal::IsHandleValid(MaterialHandle))
 			{
-				return GetMaterialHandles();
-			}
+				if (RefreshMaterials())
+				{
+					return GetMaterialHandles();
+				}
 
-			return { };
-		}
-		else
-		{
-			MtlHandles.Emplace(MaterialHandle.Value);
+				return { };
+			}
+		
+			MtlHandles.Emplace(MaterialHandle);
 		}
 	}
 	
@@ -338,6 +340,7 @@ bool FAvaMaskAvaShapeMaterialCollectionHandle::RefreshMaterials()
 	UAvaObjectHandleSubsystem* HandleSubsystem = UE::Ava::Internal::GetObjectHandleSubsystem();
 	if (!HandleSubsystem)
 	{
+		UE_LOG(LogAvaMask, Error, TEXT("ObjectHandleSubsystem was invalid"));
 		return false;
 	}
 	
diff --git a/Engine/Plugins/Experimental/Avalanche/Source/AvalancheMask/Private/Handling/AvaMaskParametricMaterialHandle.cpp b/Engine/Plugins/Experimental/Avalanche/Source/AvalancheMask/Private/Handling/AvaMaskParametricMaterialHandle.cpp
--- a/Engine/Plugins/Experimental/Avalanche/Source/AvalancheMask/Private/Handling/AvaMaskParametricMaterialHandle.cpp
+++ b/Engine/Plugins/Experimental/Avalanche/Source/AvalancheMask/Private/Handling/AvaMaskParametricMaterialHandle.cpp
@@ -69,7 +69,13 @@ bool FAvaMaskParametricMaterialHandle::SetMaskParameters(
 {
 	if (UMaterialInstanceDynamic* MaterialInstance = GetMaterialInstance())
 	{
-		UE_LOG(LogAvaMask, Display, TEXT("SetParameters: Texture:%s, Channel:%s"), InTexture ? *InTexture->GetName() : TEXT("(None)"), *UE::AvaMask::Internal::MaskChannelEnumToVector[InChannel].ToString());
+		FLinearColor PaddingValueV = FLinearColor(FMath::Max(0, InPadding.X), FMath::Max(0, InPadding.Y), 0, 0);
+		FLinearColor FeatherValueV = FLinearColor(bInApplyFeathering ? 1.0f : 0.0f, FMath::Max(0, InOuterFeathering), FMath::Max(0, InInnerFeathering), FMath::Max(0, FMath::Max(InOuterFeathering, InInnerFeathering)));
+		
+		UE_LOG(LogAvaMask, Display, TEXT("SetParameters:\nTexture:%s\nChannel:%s\nFeather:%s\n"),
+			InTexture ? *InTexture->GetName() : TEXT("(None)"),
+			*UE::AvaMask::Internal::MaskChannelEnumToVector[InChannel].ToString(),
+			*FeatherValueV.ToString());
 
 		UTexture* DefaultTexture = nullptr;
 		MaterialInstance->GetTextureParameterDefaultValue(UE::AvaMask::Internal::TextureParameterInfo, DefaultTexture);
@@ -80,8 +86,8 @@ bool FAvaMaskParametricMaterialHandle::SetMaskParameters(
 		MaterialInstance->SetVectorParameterValueByInfo(UE::AvaMask::Internal::ChannelParameterInfo, UE::AvaMask::Internal::MaskChannelEnumToVector[InChannel]);
 		MaterialInstance->SetScalarParameterValueByInfo(UE::AvaMask::Internal::InvertParameterInfo, bInInverted ? 1.0f : 0.0f);
 		MaterialInstance->SetScalarParameterValueByInfo(UE::AvaMask::Internal::BaseOpacityParameterInfo, FMath::Clamp(InBaseOpacity, 0.0f, 1.0f));
-		MaterialInstance->SetVectorParameterValueByInfo(UE::AvaMask::Internal::PaddingParameterInfo, FLinearColor(FMath::Max(0, InPadding.X), FMath::Max(0, InPadding.Y), 0, 0));
-		MaterialInstance->SetVectorParameterValueByInfo(UE::AvaMask::Internal::FeatherParameterInfo, FLinearColor(bInApplyFeathering ? 1.0f : 0.0f, FMath::Max(0, InOuterFeathering), FMath::Max(0, InInnerFeathering), FMath::Max(0, FMath::Max(InOuterFeathering, InInnerFeathering))));
+		MaterialInstance->SetVectorParameterValueByInfo(UE::AvaMask::Internal::PaddingParameterInfo, PaddingValueV);
+		MaterialInstance->SetVectorParameterValueByInfo(UE::AvaMask::Internal::FeatherParameterInfo, FeatherValueV);
 
 		return true;
 	}
diff --git a/Engine/Plugins/Experimental/Avalanche/Source/AvalancheMask/Private/Handling/IAvaMaskMaterialCollectionHandle.h b/Engine/Plugins/Experimental/Avalanche/Source/AvalancheMask/Private/Handling/IAvaMaskMaterialCollectionHandle.h
--- a/Engine/Plugins/Experimental/Avalanche/Source/AvalancheMask/Private/Handling/IAvaMaskMaterialCollectionHandle.h
+++ b/Engine/Plugins/Experimental/Avalanche/Source/AvalancheMask/Private/Handling/IAvaMaskMaterialCollectionHandle.h
@@ -317,6 +317,11 @@ bool TAvaMaskMaterialCollectionHandle<HandleDataType>::ValidateMaterials(FText&
 		, const bool bInIsSlotOccupied
 		, const TSharedPtr<IAvaMaskMaterialHandle>& InMaterialHandle)
 	{
+		if (!bInIsSlotOccupied)
+		{
+			return true;
+		}
+
 		TArray<FString> MissingParameterNames;
 		if (!InMaterialHandle->HasRequiredParameters(MissingParameterNames))
 		{
diff --git a/Engine/Plugins/Experimental/Avalanche/Source/AvalancheMask/Private/Mask2D/AvaMask2DBaseModifier.cpp b/Engine/Plugins/Experimental/Avalanche/Source/AvalancheMask/Private/Mask2D/AvaMask2DBaseModifier.cpp
--- a/Engine/Plugins/Experimental/Avalanche/Source/AvalancheMask/Private/Mask2D/AvaMask2DBaseModifier.cpp
+++ b/Engine/Plugins/Experimental/Avalanche/Source/AvalancheMask/Private/Mask2D/AvaMask2DBaseModifier.cpp
@@ -6,8 +6,8 @@
 #include "AvaMaskSubsystem.h"
 #include "AvaMaskUtilities.h"
 #include "Components/PrimitiveComponent.h"
+#include "Engine/CanvasRenderTarget2D.h"
 #include "Engine/Engine.h"
-#include "Engine/TextureRenderTarget2D.h"
 #include "Framework/AvaGizmoComponent.h"
 #include "GeometryMaskCanvas.h"
 #include "GeometryMaskReadComponent.h"
@@ -187,6 +187,30 @@ void UAvaMask2DBaseModifier::OnModifierAdded(EActorModifierCoreEnableReason InRe
 	}
 }
 
+void UAvaMask2DBaseModifier::OnModifierRemoved(EActorModifierCoreDisableReason InReason)
+{
+	Super::OnModifierRemoved(InReason);
+
+	TRACE_BOOKMARK(TEXT("UAvaMask2DModifier::OnModifierRemoved"));
+
+	TArray<AActor*> ActorDataActors;
+	ActorDataActors.Reserve(ActorData.Num());	
+	Algo::Transform(ActorData, ActorDataActors, [](const TPair<TWeakObjectPtr<AActor>, FAvaMask2DActorData>& InActorDataPair)
+	{
+		return InActorDataPair.Key.Get();
+	});
+	
+	for (AActor* Actor : ActorDataActors)
+	{
+		if (!Actor)
+		{
+			continue;
+		}
+		
+		RemoveFromActor(Actor);
+	}
+}
+
 void UAvaMask2DBaseModifier::OnModifiedActorTransformed()
 {
 	// This is needed to override parent behavior (to disable it)
diff --git a/Engine/Plugins/Experimental/Avalanche/Source/AvalancheMask/Public/Mask2D/AvaMask2DBaseModifier.h b/Engine/Plugins/Experimental/Avalanche/Source/AvalancheMask/Public/Mask2D/AvaMask2DBaseModifier.h
--- a/Engine/Plugins/Experimental/Avalanche/Source/AvalancheMask/Public/Mask2D/AvaMask2DBaseModifier.h
+++ b/Engine/Plugins/Experimental/Avalanche/Source/AvalancheMask/Public/Mask2D/AvaMask2DBaseModifier.h
@@ -105,6 +105,7 @@ protected:
 	//~ Begin UActorModifierCoreBase
 	virtual void OnModifierCDOSetup(FActorModifierCoreMetadata& InMetadata) override;
 	virtual void OnModifierAdded(EActorModifierCoreEnableReason InReason) override;
+	virtual void OnModifierRemoved(EActorModifierCoreDisableReason InReason) override;
 	virtual void OnModifiedActorTransformed() override;
 	virtual void SavePreState() override;
 	virtual void RestorePreState() override;
diff --git a/Engine/Plugins/Experimental/GeometryMask/Shaders/Private/GeometryMaskJFInitCS.usf b/Engine/Plugins/Experimental/GeometryMask/Shaders/Private/GeometryMaskJFInitCS.usf
--- a/Engine/Plugins/Experimental/GeometryMask/Shaders/Private/GeometryMaskJFInitCS.usf
+++ b/Engine/Plugins/Experimental/GeometryMask/Shaders/Private/GeometryMaskJFInitCS.usf
@@ -2,9 +2,6 @@
 
 #include "/Engine/Public/Platform.ush"
 
-#define TILE_SIZE 8
-#define KERNEL_SIZE 3
-
 uint2 InputDimensions;
 float2 OneOverInputDimensions;
 float2 UVRatioAdjustment;
diff --git a/Engine/Plugins/Experimental/GeometryMask/Shaders/Private/GeometryMaskJFStepCS.usf b/Engine/Plugins/Experimental/GeometryMask/Shaders/Private/GeometryMaskJFStepCS.usf
--- a/Engine/Plugins/Experimental/GeometryMask/Shaders/Private/GeometryMaskJFStepCS.usf
+++ b/Engine/Plugins/Experimental/GeometryMask/Shaders/Private/GeometryMaskJFStepCS.usf
@@ -2,10 +2,6 @@
 
 #include "/Engine/Public/Platform.ush"
 
-#define TILE_SIZE 8
-#define KERNEL_SIZE 3
-#define MAX_STEPS 13
-
 uint2 InputDimensions;
 float2 OneOverInputDimensions;
 float2 UVRatioAdjustment;
@@ -54,24 +50,24 @@ void MainCS(
 	float Y = (float)PixelCoord.y + 0.5;
 	const float2 UV = float2(X, Y) * OneOverInputDimensions * UVRatioAdjustment;
 
-    // Perform on each of R,G,B,A
-    [unroll]
-    for (int ChannelIdx = 0; ChannelIdx < NUM_CHANNELS; ++ChannelIdx)
-    {
-    	float4 Input = InputBuffer[PixelCoord1D + ChannelIdx];
-    	const uint2 StepSizeWithChannel = StepSize[ChannelIdx] * uint2(1, 1);
+	// Perform on each of R,G,B,A
+	[unroll]
+	for (int ChannelIdx = 0; ChannelIdx < NUM_CHANNELS; ++ChannelIdx)
+	{
+		float4 Input = InputBuffer[PixelCoord1D + ChannelIdx];
+		const uint2 StepSizeWithChannel = StepSize[ChannelIdx] * uint2(1, 1);
 
-    	// Check if we should calculate this channel (for this step)
-    	if (ChannelMask[ChannelIdx] == 0)
-    	{
-    		MinSample[ChannelIdx] = Input[ChannelIdx];
-    		continue;
-    	}
+		// Check if we should calculate this channel (for this step)
+		if (ChannelMask[ChannelIdx] == 0)
+		{
+			MinSample[ChannelIdx] = Input[ChannelIdx];
+			continue;
+		}
 
-    	MinDistance[ChannelIdx] = 1000000000;
-    	MinSample[ChannelIdx] = float4(0,0,0,0);
+		MinDistance[ChannelIdx] = 1000000000;
+		MinSample[ChannelIdx] = float4(0,0,0,0);
 
-    	[unroll]
+		[unroll]
 		for (int HorizontalIdx = 0; HorizontalIdx < 3; ++HorizontalIdx)
 		{
 			[unroll]
@@ -79,10 +75,9 @@ void MainCS(
 			{
 				uint NeighborIdx = (VerticalIdx * KERNEL_SIZE) + HorizontalIdx;
 				uint2 NeighborPixelCoord = PixelCoord + (NeighborOffsetIndices[NeighborIdx] * StepSizeWithChannel);
-				if (NeighborPixelCoord.x > InputDimensions.x - 1 || NeighborPixelCoord.y > InputDimensions.y - 1)
-				{
-					continue;
-				}
+
+				// Prevents border banding
+				NeighborPixelCoord = min(InputDimensions - 1, max(uint2(1,1), NeighborPixelCoord));
 
 				uint NeighborPixelCoord1D = (NeighborPixelCoord.y * RowLength) + (NeighborPixelCoord.x * NUM_CHANNELS);
 
@@ -101,7 +96,7 @@ void MainCS(
 			}
 		}
 
-    	MinSample[ChannelIdx].a = Input.a;
-    	OutputBuffer[PixelCoord1D + ChannelIdx] = MinSample[ChannelIdx];
-    }
+		MinSample[ChannelIdx].a = Input.a;
+		OutputBuffer[PixelCoord1D + ChannelIdx] = MinSample[ChannelIdx];
+	}
 }
diff --git a/Engine/Plugins/Experimental/GeometryMask/Shaders/Private/GeometryMaskJFtoDFCS.usf b/Engine/Plugins/Experimental/GeometryMask/Shaders/Private/GeometryMaskJFtoDFCS.usf
--- a/Engine/Plugins/Experimental/GeometryMask/Shaders/Private/GeometryMaskJFtoDFCS.usf
+++ b/Engine/Plugins/Experimental/GeometryMask/Shaders/Private/GeometryMaskJFtoDFCS.usf
@@ -2,7 +2,7 @@
 
 #include "/Engine/Public/Platform.ush"
 
-#define TILE_SIZE 8
+#define DEBUG_SOBEL 0
 
 uint2 InputDimensions;
 float2 OneOverInputDimensions;
@@ -36,6 +36,10 @@ void MainCS(
 	float4 Outline = OriginalInputTexture.Load(int3(PixelCoord, 0));
 	float4 Mask = float4(1,1,1,1);
 
+#if DEBUG_SOBEL
+	float4 Sample = InputBuffer[PixelCoord1D];
+	Outline[0] = saturate(distance(UV, Sample.xy));
+#else
 	// Perform on each of R,G,B,A
 	[unroll]
 	for (int ChannelIdx = 0; ChannelIdx < NUM_CHANNELS; ++ChannelIdx)
@@ -43,14 +47,16 @@ void MainCS(
 		const uint OutputChannelIdx = ChannelMap[ChannelIdx];
 		
 		float4 Sample = InputBuffer[PixelCoord1D + ChannelIdx];
+
 		Mask[OutputChannelIdx] = (Sample.x + Sample.y + Sample.a > 0.001);
 
 		Outline[OutputChannelIdx] = saturate(distance(UV, Sample.xy) * StepDistanceMultipliers[OutputChannelIdx]);
 
-		// Set range such that inner going from 0.0-0.5, outer goes from 0.5-1.0
+		// Set range such that inner goes from 0.0-0.5, outer goes from 0.5-1.0
 		Outline[OutputChannelIdx] = mad(Outline[OutputChannelIdx], 0.5, 0.5);
 		Outline[OutputChannelIdx] = 1.0 - abs(Outline[OutputChannelIdx] - saturate(Sample.a));
 	}
+#endif
 
 	Outline = lerp(0, Outline, Mask);
 
diff --git a/Engine/Plugins/Experimental/GeometryMask/Source/GeometryMask/Private/GeometryMaskCanvas.cpp b/Engine/Plugins/Experimental/GeometryMask/Source/GeometryMask/Private/GeometryMaskCanvas.cpp
--- a/Engine/Plugins/Experimental/GeometryMask/Source/GeometryMask/Private/GeometryMaskCanvas.cpp
+++ b/Engine/Plugins/Experimental/GeometryMask/Source/GeometryMask/Private/GeometryMaskCanvas.cpp
@@ -4,10 +4,11 @@
 
 #include "Algo/RemoveIf.h"
 #include "CanvasTypes.h"
+#include "Engine/CanvasRenderTarget2D.h"
 #include "Engine/Engine.h"
-#include "Engine/TextureRenderTarget2D.h"
 #include "Engine/World.h"
 #include "GeometryMaskCanvasResource.h"
+#include "GeometryMaskModule.h"
 #include "IGeometryMaskWriteInterface.h"
 
 void UGeometryMaskCanvas::BeginDestroy()
@@ -65,7 +66,7 @@ void UGeometryMaskCanvas::RemoveWriter(const TScriptInterface<IGeometryMaskWrite
 	Writers.Remove(InWriter.GetObject());
 }
 
-UTextureRenderTarget2D* UGeometryMaskCanvas::GetTexture() const
+UCanvasRenderTarget2D* UGeometryMaskCanvas::GetTexture() const
 {
 	if (CanvasResource)
 	{
@@ -179,6 +180,7 @@ void UGeometryMaskCanvas::OnDrawToCanvas(FCanvas* InCanvas)
 		if (Writer.IsValid())
 		{
 			TRACE_CPUPROFILER_EVENT_SCOPE(UGeometryMaskCanvas::OnDrawToCanvas);
+			DECLARE_SCOPE_CYCLE_COUNTER(TEXT("UGeometryMaskCanvas::OnDrawToCanvas"), STAT_GeometryMask_Update, STATGROUP_GeometryMask);
 			Writer->DrawToCanvas(InCanvas);	
 		}
 	}
diff --git a/Engine/Plugins/Experimental/GeometryMask/Source/GeometryMask/Private/GeometryMaskCanvasActor.cpp b/Engine/Plugins/Experimental/GeometryMask/Source/GeometryMask/Private/GeometryMaskCanvasActor.cpp
--- a/Engine/Plugins/Experimental/GeometryMask/Source/GeometryMask/Private/GeometryMaskCanvasActor.cpp
+++ b/Engine/Plugins/Experimental/GeometryMask/Source/GeometryMask/Private/GeometryMaskCanvasActor.cpp
@@ -23,7 +23,7 @@ AGeometryMaskCanvasActor::AGeometryMaskCanvasActor(const FObjectInitializer& Obj
 	Canvas = nullptr;
 }
 
-UTextureRenderTarget2D* AGeometryMaskCanvasActor::GetTexture()
+UCanvasRenderTarget2D* AGeometryMaskCanvasActor::GetTexture()
 {
 	return Canvas ? Canvas->GetTexture() : nullptr;
 }
diff --git a/Engine/Plugins/Experimental/GeometryMask/Source/GeometryMask/Private/GeometryMaskCanvasResource.cpp b/Engine/Plugins/Experimental/GeometryMask/Source/GeometryMask/Private/GeometryMaskCanvasResource.cpp
--- a/Engine/Plugins/Experimental/GeometryMask/Source/GeometryMask/Private/GeometryMaskCanvasResource.cpp
+++ b/Engine/Plugins/Experimental/GeometryMask/Source/GeometryMask/Private/GeometryMaskCanvasResource.cpp
@@ -3,7 +3,7 @@
 #include "GeometryMaskCanvasResource.h"
 
 #include "Engine/Canvas.h"
-#include "Engine/TextureRenderTarget2D.h"
+#include "Engine/CanvasRenderTarget2D.h"
 #include "Engine/World.h"
 #include "GeometryMaskModule.h"
 #include "GeometryMaskSettings.h"
@@ -17,6 +17,8 @@
 
 namespace UE::GeometryMask::Private
 {
+	static constexpr ETextureRenderTargetSampleCount RenderTargetSampleCount = ETextureRenderTargetSampleCount::RTSC_4;
+	
 	void OverscanProjectionMatrix(FMatrix& InOutMatrix, const FIntPoint& InSize, const int32 InPadding)
 	{
 		const FVector2f Multiplier(
@@ -42,14 +44,24 @@ UGeometryMaskCanvasResource::UGeometryMaskCanvasResource()
 	};
 
 	FGeometryMaskPostProcessParameters_Blur PostProcessParameters_Blur;
+	PostProcessParameters_Blur.bPerChannelApplyBlur.SetRange(0, 4, false);
 	PostProcessParameters_Blur.PerChannelBlurStrength = { 16, 16, 16, 16};
 	PostProcess_Blur = MakeShared<FGeometryMaskPostProcess_Blur>(PostProcessParameters_Blur);
 
 	FGeometryMaskPostProcessParameters_DistanceField PostProcessParameters_DistanceField;
-	PostProcessParameters_DistanceField.bPerChannelCalculateDF.SetRange(0, 4, true);
+	PostProcessParameters_DistanceField.bPerChannelCalculateDF.SetRange(0, 4, false);
 	PostProcess_DistanceField = MakeShared<FGeometryMaskPostProcess_DistanceField>(PostProcessParameters_DistanceField);
 }
 
+UGeometryMaskCanvasResource::~UGeometryMaskCanvasResource()
+{
+	if (!IsUnreachable() && IsValid(CanvasObject) && CanvasObject->Canvas)
+	{
+		delete CanvasObject->Canvas;
+		CanvasObject->Canvas = nullptr;
+	}
+}
+
 const EGeometryMaskColorChannel UGeometryMaskCanvasResource::GetNextAvailableColorChannel() const
 {
 	for (const TPair<EGeometryMaskColorChannel, FName>& ColorChannelCanvas : DependentCanvasNames)
@@ -89,6 +101,7 @@ bool UGeometryMaskCanvasResource::Checkin(const FName InRequestingCanvasName)
 		{
 			// Effectively free this ColorChannel and make available for Checkout
 			ColorChannelCanvas.Value = NAME_None;
+			ResetRenderParameters(ColorChannelCanvas.Key);
 			return true;
 		}
 	}
@@ -100,7 +113,8 @@ bool UGeometryMaskCanvasResource::Checkin(const FName InRequestingCanvasName)
 void UGeometryMaskCanvasResource::UpdateViewportSize()
 {
 	if (ViewportSize.Size() > 0
-		&& RenderTargetTexture)
+		&& IsValid(RenderTargetTexture)
+		&& !IsUnreachable())
 	{
 		const float SizeMultiplier = GetDefault<UGeometryMaskSettings>()->GetDefaultResolutionMultiplier();
 		
@@ -127,6 +141,12 @@ void UGeometryMaskCanvasResource::UpdateViewportSize()
 			SizeX /= RatioY;
 		}
 
+		if (RenderTargetTexture->SizeX == SizeX
+			&& RenderTargetTexture->SizeY == SizeY)
+		{
+			return;
+		}
+
 		// Update RT size to viewport size
 		RenderTargetTexture->ResizeTarget(SizeX, SizeY);
 	}
@@ -200,11 +220,29 @@ void UGeometryMaskCanvasResource::UpdateRenderParameters(
 		PostProcess_DistanceField->SetParameters(DFParameters);
 	}
 
-	// Effects may have changed viewport padding
-	UpdateViewportSize();
-
 	bApplyBlur = BlurParameters.bPerChannelApplyBlur.CountSetBits(0) > 0; // if any channels have blur
 	bApplyDF = DFParameters.bPerChannelCalculateDF.CountSetBits(0) > 0; // if any channels have DF
+
+	if (UCanvasRenderTarget2D* Texture = GetRenderTargetTexture())
+	{
+		if (bInApplyBlur || bInApplyFeather)
+		{
+			// MSAA not supported for these effects (they need UAV access)
+			Texture->SetSampleCount(ETextureRenderTargetSampleCount::RTSC_1);
+		}
+		else
+		{
+			Texture->SetSampleCount(UE::GeometryMask::Private::RenderTargetSampleCount);
+		}
+	}
+
+	// Effects may have changed viewport padding
+    UpdateViewportSize();
+}
+
+void UGeometryMaskCanvasResource::ResetRenderParameters(EGeometryMaskColorChannel InColorChannel)
+{
+	UpdateRenderParameters(InColorChannel, false, 0.0, false, 0, 0);
 }
 
 void UGeometryMaskCanvasResource::SetViewportSize(const FIntPoint& InViewportSize)
@@ -243,17 +281,18 @@ int32 UGeometryMaskCanvasResource::GetViewportPadding() const
 		}
 	}
 
-	
 	return FMath::Max(MaxFeatherRadius, MaxBlurRadius);
 }
 
-UTextureRenderTarget2D* UGeometryMaskCanvasResource::GetRenderTargetTexture()
+UCanvasRenderTarget2D* UGeometryMaskCanvasResource::GetRenderTargetTexture()
 {
-	if (!RenderTargetTexture)
+	if (!IsValid(RenderTargetTexture))
 	{
-		const FName ObjectName = MakeUniqueObjectName(this, UTextureRenderTarget2D::StaticClass(), FName(FString::Printf(TEXT("GeometryMaskCanvasResource_RenderTarget"))));
-		RenderTargetTexture = NewObject<UTextureRenderTarget2D>(this, ObjectName);
+		const FName ObjectName = MakeUniqueObjectName(this, UCanvasRenderTarget2D::StaticClass(), FName(FString::Printf(TEXT("GeometryMaskCanvasResource_RenderTarget"))));
+		RenderTargetTexture = NewObject<UCanvasRenderTarget2D>(this, ObjectName);
 		RenderTargetTexture->bForceLinearGamma = false;
+		RenderTargetTexture->bAutoGenerateMips = false;
+		RenderTargetTexture->SetSampleCount(UE::GeometryMask::Private::RenderTargetSampleCount);
 		RenderTargetTexture->InitAutoFormat(ViewportSize.X, ViewportSize.Y);
 	}
 
@@ -274,11 +313,22 @@ void UGeometryMaskCanvasResource::Draw(UWorld* InWorld, FSceneView& InView)
 		return;
 	}
 
-	if (UTextureRenderTarget2D* Texture = GetRenderTargetTexture())
+	if (UCanvasRenderTarget2D* Texture = GetRenderTargetTexture())
 	{
 		if (!CanvasObject)
 		{
 			CanvasObject = NewObject<UCanvas>(GetTransientPackage());
+
+			FTextureRenderTargetResource* RenderTargetResource = Texture->GameThread_GetRenderTargetResource();
+			FCanvas* NewCanvas = new FCanvas(
+				RenderTargetResource,
+				nullptr,
+				InWorld,
+				InWorld->GetFeatureLevel(),
+				// Draw immediately so that interleaved SetVectorParameter (etc) function calls work as expected
+				FCanvas::CDM_ImmediateDrawing);
+
+			CanvasObject->Init(Texture->SizeX, Texture->SizeY, &InView, NewCanvas);
 		}
 
 		// Begin
@@ -294,16 +344,8 @@ void UGeometryMaskCanvasResource::Draw(UWorld* InWorld, FSceneView& InView)
 			
 			InWorld->FlushDeferredParameterCollectionInstanceUpdates();
 
-			FTextureRenderTargetResource* RenderTargetResource = Texture->GameThread_GetRenderTargetResource();
-			FCanvas* NewCanvas = new FCanvas(
-				RenderTargetResource,
-				nullptr,
-				InWorld,
-				InWorld->GetFeatureLevel(),
-				// Draw immediately so that interleaved SetVectorParameter (etc) function calls work as expected
-				FCanvas::CDM_ImmediateDrawing);
-
-			CanvasObject->Init(Texture->SizeX, Texture->SizeY, &InView, NewCanvas);
+			CanvasObject->SizeX = Texture->SizeX;
+			CanvasObject->SizeY = Texture->SizeY;
 			CanvasObject->Update();
 			CanvasObject->SetView(&InView);
 			CanvasObject->Canvas->Clear(GetRenderTargetTexture()->ClearColor);
@@ -346,9 +388,6 @@ void UGeometryMaskCanvasResource::Draw(UWorld* InWorld, FSceneView& InView)
 						PostProcess_Blur->Execute(RenderTarget);
 					}
 				}
-
-				delete CanvasObject->Canvas;
-				CanvasObject->Canvas = nullptr;
 			}
 		}
 	}
diff --git a/Engine/Plugins/Experimental/GeometryMask/Source/GeometryMask/Private/GeometryMaskModule.cpp b/Engine/Plugins/Experimental/GeometryMask/Source/GeometryMask/Private/GeometryMaskModule.cpp
--- a/Engine/Plugins/Experimental/GeometryMask/Source/GeometryMask/Private/GeometryMaskModule.cpp
+++ b/Engine/Plugins/Experimental/GeometryMask/Source/GeometryMask/Private/GeometryMaskModule.cpp
@@ -2,6 +2,7 @@
 
 #include "GeometryMaskModule.h"
 
+#include "GeometryMaskSettings.h"
 #include "IGeometryMaskModule.h"
 #include "Interfaces/IPluginManager.h"
 #include "Misc/Paths.h"
@@ -20,8 +21,58 @@ IGeometryMaskModule& IGeometryMaskModule::Get()
 void FGeometryMaskModule::StartupModule()
 {
 	const FString PluginShaderDir = FPaths::Combine(IPluginManager::Get().FindPlugin(UE_PLUGIN_NAME)->GetBaseDir(), TEXT("Shaders"));
-	// @todo: uncomment when shaders implemented
 	AddShaderSourceDirectoryMapping(TEXT("/Plugin/GeometryMask"), PluginShaderDir);
+
+	if (GIsEditor && !IsRunningCommandlet())
+	{
+		RegisterConsoleCommands();
+	}
+}
+
+void FGeometryMaskModule::ShutdownModule()
+{
+	// Cleanup commands
+	UnregisterConsoleCommands();
+}
+
+void FGeometryMaskModule::RegisterConsoleCommands()
+{
+	ConsoleCommands.Add(IConsoleManager::Get().RegisterConsoleCommand(
+		TEXT("GeometryMask.DebugDF"),
+		TEXT("0: Off, 1: Sobel, 2...StepN: Step0...StepN"),
+		FConsoleCommandWithArgsDelegate::CreateRaw(this, &FGeometryMaskModule::ExecuteDebugDF),
+		ECVF_Default
+	));
+}
+
+void FGeometryMaskModule::UnregisterConsoleCommands()
+{
+	for (IConsoleObject* Cmd : ConsoleCommands)
+	{
+		IConsoleManager::Get().UnregisterConsoleObject(Cmd);
+	}
+
+	ConsoleCommands.Empty();
+}
+
+void FGeometryMaskModule::ExecuteDebugDF(const TArray<FString>& InArgs)
+{
+#if WITH_EDITORONLY_DATA
+	int32 Value = 0;
+	
+	UGeometryMaskSettings* Settings = GetMutableDefault<UGeometryMaskSettings>();
+	check(Settings);
+	
+	if (!InArgs.IsEmpty())
+	{
+		if (InArgs[0].IsNumeric())
+		{
+			Value = FCString::Atoi(*InArgs[0]);
+		}
+	}
+
+	Settings->DebugDF = Value;
+#endif
 }
 
 #undef LOCTEXT_NAMESPACE
diff --git a/Engine/Plugins/Experimental/GeometryMask/Source/GeometryMask/Private/GeometryMaskModule.h b/Engine/Plugins/Experimental/GeometryMask/Source/GeometryMask/Private/GeometryMaskModule.h
--- a/Engine/Plugins/Experimental/GeometryMask/Source/GeometryMask/Private/GeometryMaskModule.h
+++ b/Engine/Plugins/Experimental/GeometryMask/Source/GeometryMask/Private/GeometryMaskModule.h
@@ -6,8 +6,10 @@
 
 #include "Logging/LogMacros.h"
 #include "Modules/ModuleInterface.h"
+#include "Stats/Stats2.h"
 
 DECLARE_LOG_CATEGORY_EXTERN(LogGeometryMask, Log, All);
+DECLARE_STATS_GROUP(TEXT("GeometryMask"), STATGROUP_GeometryMask, STATCAT_Advanced);
 
 class FGeometryMaskModule
 	: public IGeometryMaskModule
@@ -15,5 +17,19 @@ class FGeometryMaskModule
 public:
 	// ~Begin IModuleInterface
 	virtual void StartupModule() override;
+	virtual void ShutdownModule() override;
 	// ~End IModuleInterface
+
+private:
+	/** Called from StartupModule and sets up console commands for the plugin via IConsoleManager */
+	void RegisterConsoleCommands();
+	
+	/** Called from ShutdownModule and clears out previously registered console commands */
+	void UnregisterConsoleCommands();
+
+	void ExecuteDebugDF(const TArray<FString>& InArgs);
+
+private:
+	/** References of registered console commands via IConsoleManager */
+	TArray<IConsoleObject*> ConsoleCommands;
 };
diff --git a/Engine/Plugins/Experimental/GeometryMask/Source/GeometryMask/Private/GeometryMaskSVE.cpp b/Engine/Plugins/Experimental/GeometryMask/Source/GeometryMask/Private/GeometryMaskSVE.cpp
--- a/Engine/Plugins/Experimental/GeometryMask/Source/GeometryMask/Private/GeometryMaskSVE.cpp
+++ b/Engine/Plugins/Experimental/GeometryMask/Source/GeometryMask/Private/GeometryMaskSVE.cpp
@@ -13,12 +13,13 @@ FGeometryMaskSceneViewExtension::FGeometryMaskSceneViewExtension(
 	UWorld* InWorld)
 	: FWorldSceneViewExtension(AutoRegister, InWorld)
 {
+	GeometryMaskSubsystemWeak = GEngine->GetEngineSubsystem<UGeometryMaskSubsystem>();
 }
 
 void FGeometryMaskSceneViewExtension::BeginRenderViewFamily(
 	FSceneViewFamily& InViewFamily)
 {
-	if (UGeometryMaskSubsystem* Subsystem = GEngine->GetEngineSubsystem<UGeometryMaskSubsystem>())
+	if (UGeometryMaskSubsystem* Subsystem = GeometryMaskSubsystemWeak.Get())
 	{
 		Subsystem->Update(GetWorld(), InViewFamily);
 	}
diff --git a/Engine/Plugins/Experimental/GeometryMask/Source/GeometryMask/Private/GeometryMaskSVE.h b/Engine/Plugins/Experimental/GeometryMask/Source/GeometryMask/Private/GeometryMaskSVE.h
--- a/Engine/Plugins/Experimental/GeometryMask/Source/GeometryMask/Private/GeometryMaskSVE.h
+++ b/Engine/Plugins/Experimental/GeometryMask/Source/GeometryMask/Private/GeometryMaskSVE.h
@@ -4,6 +4,7 @@
 
 #include "SceneViewExtension.h"
 
+class UGeometryMaskSubsystem;
 class UGeometryMaskCanvas;
 
 class FGeometryMaskSceneViewExtension
@@ -18,4 +19,7 @@ public:
 	virtual void BeginRenderViewFamily(FSceneViewFamily& InViewFamily) override;
 	virtual bool IsActiveThisFrame_Internal(const FSceneViewExtensionContext& Context) const override;
 	// ~End ISceneViewExtension
+
+private:
+	TWeakObjectPtr<UGeometryMaskSubsystem> GeometryMaskSubsystemWeak;
 };
diff --git a/Engine/Plugins/Experimental/GeometryMask/Source/GeometryMask/Private/GeometryMaskSettings.h b/Engine/Plugins/Experimental/GeometryMask/Source/GeometryMask/Private/GeometryMaskSettings.h
--- a/Engine/Plugins/Experimental/GeometryMask/Source/GeometryMask/Private/GeometryMaskSettings.h
+++ b/Engine/Plugins/Experimental/GeometryMask/Source/GeometryMask/Private/GeometryMaskSettings.h
@@ -11,6 +11,12 @@ class UGeometryMaskSettings : public UDeveloperSettings
 {
 	GENERATED_BODY()
 
+public:
+#if WITH_EDITORONLY_DATA
+	UPROPERTY()
+	uint8 DebugDF = 0;
+#endif
+	
 public:
 	virtual FName GetContainerName() const override { return FName(TEXT("Project")); }
 	virtual FName GetCategoryName() const override { return FName(TEXT("Plugins")); }
diff --git a/Engine/Plugins/Experimental/GeometryMask/Source/GeometryMask/Private/GeometryMaskSubsystem.cpp b/Engine/Plugins/Experimental/GeometryMask/Source/GeometryMask/Private/GeometryMaskSubsystem.cpp
--- a/Engine/Plugins/Experimental/GeometryMask/Source/GeometryMask/Private/GeometryMaskSubsystem.cpp
+++ b/Engine/Plugins/Experimental/GeometryMask/Source/GeometryMask/Private/GeometryMaskSubsystem.cpp
@@ -2,7 +2,6 @@
 
 #include "GeometryMaskSubsystem.h"
 
-#include "Algo/RemoveIf.h"
 #include "Engine/Engine.h"
 #include "Engine/GameViewportClient.h"
 #include "Engine/LocalPlayer.h"
@@ -10,6 +9,7 @@
 #include "GameFramework/PlayerController.h"
 #include "GeometryMaskCanvas.h"
 #include "GeometryMaskCanvasResource.h"
+#include "GeometryMaskModule.h"
 #include "SceneView.h"
 #include "UnrealClient.h"
 
@@ -51,6 +51,11 @@ TArray<FName> UGeometryMaskSubsystem::GetCanvasNames()
 	return {};
 }
 
+int32 UGeometryMaskSubsystem::GetNumActiveCanvasResources() const
+{
+	return NamedCanvases.Num();
+}
+
 const TArray<TObjectPtr<UGeometryMaskCanvasResource>>& UGeometryMaskSubsystem::GetCanvasResources() const
 {
 	return CanvasResources;
@@ -60,16 +65,26 @@ void UGeometryMaskSubsystem::Update(
 	UWorld* InWorld,
 	FSceneViewFamily& InViewFamily)
 {
+	if (!bDoUpdates)
+	{
+		return;
+	}
+	
 	if (UGeometryMaskSubsystem* Subsystem = GEngine->GetEngineSubsystem<UGeometryMaskSubsystem>())
 	{
+		DECLARE_SCOPE_CYCLE_COUNTER(TEXT("UGeometryMaskSubsystem::Update"), STAT_GeometryMask_UpdateAll, STATGROUP_GeometryMask);
+	
 		for (const FSceneView*& View : InViewFamily.Views)
 		{
 			FSceneView* MutableView = const_cast<FSceneView*>(View);
+
+			// Removes invalid writers
 			for (const TPair<FName, TObjectPtr<UGeometryMaskCanvas>>& NamedCanvas : Subsystem->NamedCanvases)
 			{
 				NamedCanvas.Value->Update(InWorld, *MutableView);			
 			}
 
+			// Updates the texture resource
 			for (const TObjectPtr<UGeometryMaskCanvasResource>& Resource : CanvasResources)
 			{
 				Resource->Update(InWorld, *MutableView);
@@ -78,6 +93,16 @@ void UGeometryMaskSubsystem::Update(
 	}
 }
 
+void UGeometryMaskSubsystem::ToggleUpdate(const TOptional<bool>& bInShouldUpdate)
+{
+	// New value should be user provided, or the inverse of the existing value (toggle)
+	bool bShouldUpdate = bInShouldUpdate.Get(!bDoUpdates);
+	if (bDoUpdates != bShouldUpdate)
+	{
+		bDoUpdates = bShouldUpdate;
+	}
+}
+
 int32 UGeometryMaskSubsystem::RemoveWithoutWriters()
 {
 	int32 NumRemoved = 0;
diff --git a/Engine/Plugins/Experimental/GeometryMask/Source/GeometryMask/Private/GeometryMaskTypes.cpp b/Engine/Plugins/Experimental/GeometryMask/Source/GeometryMask/Private/GeometryMaskTypes.cpp
--- a/Engine/Plugins/Experimental/GeometryMask/Source/GeometryMask/Private/GeometryMaskTypes.cpp
+++ b/Engine/Plugins/Experimental/GeometryMask/Source/GeometryMask/Private/GeometryMaskTypes.cpp
@@ -38,7 +38,12 @@ namespace UE::GeometryMask
 	}
 }
 
-UTextureRenderTarget2D* UGeometryMaskCanvasReferenceComponentBase::GetTexture()
+UGeometryMaskCanvasReferenceComponentBase::~UGeometryMaskCanvasReferenceComponentBase()
+{
+	Cleanup();
+}
+
+UCanvasRenderTarget2D* UGeometryMaskCanvasReferenceComponentBase::GetTexture()
 {
 	if (const UGeometryMaskCanvas* Canvas = CanvasWeak.Get())
 	{
@@ -115,3 +120,13 @@ bool UGeometryMaskCanvasReferenceComponentBase::TryResolveNamedCanvas(FName InCa
 
 	return CanvasWeak.IsValid();
 }
+
+bool UGeometryMaskCanvasReferenceComponentBase::Cleanup()
+{
+	if (!GEngine)
+	{
+		return false;
+	}
+
+	return true;
+}
diff --git a/Engine/Plugins/Experimental/GeometryMask/Source/GeometryMask/Private/GeometryMaskWriteComponent.cpp b/Engine/Plugins/Experimental/GeometryMask/Source/GeometryMask/Private/GeometryMaskWriteComponent.cpp
--- a/Engine/Plugins/Experimental/GeometryMask/Source/GeometryMask/Private/GeometryMaskWriteComponent.cpp
+++ b/Engine/Plugins/Experimental/GeometryMask/Source/GeometryMask/Private/GeometryMaskWriteComponent.cpp
@@ -76,7 +76,7 @@ void UGeometryMaskWriteMeshComponent::DrawToCanvas(FCanvas* InCanvas)
 			const FName KeyName = KeyNames[TaskIdx];
 			FTransform LocalToWorld;
 			
-			if (const USceneComponent* Component = CachedComponents[KeyName].Get())
+			if (const USceneComponent* Component = CachedComponentsWeak[KeyName].Get())
 			{
 				LocalToWorld = Component->GetComponentToWorld();
 			}
@@ -145,11 +145,7 @@ void UGeometryMaskWriteMeshComponent::OnComponentDestroyed(bool bDestroyingHiera
 {
 	Super::OnComponentDestroyed(bDestroyingHierarchy);
 
-	if (UGeometryMaskCanvas* Canvas = CanvasWeak.Get())
-	{
-		Canvas->RemoveWriter(this);
-		CanvasWeak.Reset();
-	}
+	Cleanup();
 }
 
 void UGeometryMaskWriteMeshComponent::UpdateCachedData()
@@ -166,7 +162,7 @@ void UGeometryMaskWriteMeshComponent::UpdateCachedData()
 		}
 
 		LastPrimitiveComponentCount = PrimitiveComponents.Num();
-		CachedComponents.Reserve(LastPrimitiveComponentCount);
+		CachedComponentsWeak.Reserve(LastPrimitiveComponentCount);
 		CachedMeshData.Reserve(LastPrimitiveComponentCount);
 
 		UpdateCachedStaticMeshData(PrimitiveComponents);
@@ -256,7 +252,7 @@ void UGeometryMaskWriteMeshComponent::UpdateCachedStaticMeshData(TConstArrayView
 
 						StaticMeshObjectNames.Add(StaticMesh->GetFName());
 						StaticMeshResources.Add(&RenderData->LODResources[0]);
-						CachedComponents.Add(StaticMesh->GetFName(), StaticMeshComponent);
+						CachedComponentsWeak.Add(StaticMesh->GetFName(), StaticMeshComponent);
 					}
 				}
 			}
@@ -354,7 +350,7 @@ void UGeometryMaskWriteMeshComponent::UpdateCachedDynamicMeshData(TConstArrayVie
 				{
 					DynamicMeshObjectNames.Add(DynamicMeshComponent->GetFName());
 					DynamicMeshes.Add(DynamicMesh);
-					CachedComponents.Add(DynamicMeshComponent->GetFName(), DynamicMeshComponent);
+					CachedComponentsWeak.Add(DynamicMeshComponent->GetFName(), DynamicMeshComponent);
 				}
 			}
 		}
@@ -431,3 +427,20 @@ bool UGeometryMaskWriteMeshComponent::TryResolveCanvas()
 
 	return false;
 }
+
+bool UGeometryMaskWriteMeshComponent::Cleanup()
+{
+	if (!Super::Cleanup())
+	{
+		return false;
+	}
+
+	if (UGeometryMaskCanvas* Canvas = CanvasWeak.Get())
+	{
+		Canvas->RemoveWriter(this);
+		CanvasWeak.Reset();
+		return true;
+	}
+
+	return false;
+}
diff --git a/Engine/Plugins/Experimental/GeometryMask/Source/GeometryMask/Private/Shaders/GeometryMaskPostProcess_DistanceField.cpp b/Engine/Plugins/Experimental/GeometryMask/Source/GeometryMask/Private/Shaders/GeometryMaskPostProcess_DistanceField.cpp
--- a/Engine/Plugins/Experimental/GeometryMask/Source/GeometryMask/Private/Shaders/GeometryMaskPostProcess_DistanceField.cpp
+++ b/Engine/Plugins/Experimental/GeometryMask/Source/GeometryMask/Private/Shaders/GeometryMaskPostProcess_DistanceField.cpp
@@ -2,16 +2,22 @@
 
 #include "GeometryMaskPostProcess_DistanceField.h"
 
+#include "GeometryMaskModule.h"
 #include "GlobalShader.h"
 #include "RenderGraphBuilder.h"
 #include "RenderGraphUtils.h"
 #include "RHIStaticStates.h"
 #include "UnrealClient.h"
 
+DECLARE_GPU_STAT_NAMED(GeometryMaskJFInit, TEXT("GeometryMaskJFInit"));
+DECLARE_GPU_STAT_NAMED(GeometryMaskJFStep, TEXT("GeometryMaskJFStep"));
+DECLARE_GPU_STAT_NAMED(GeometryMaskJFtoDF, TEXT("GeometryMaskJFtoDF"));
+
 namespace UE::GeometryMask::Private
 {
 	static constexpr int32 MaxNumChannels = 4;
 	static constexpr int32 NumNeighbors = 8;
+	static constexpr int32 MaxSteps = 13;
 
 	static FIntVector4 CalculateStepCount(const FIntVector4& InRadiusSizes)
 	{
@@ -81,7 +87,9 @@ public:
 	{
 		FGlobalShader::ModifyCompilationEnvironment(InParameters, OutEnvironment);
 
+		OutEnvironment.SetDefine(TEXT("TILE_SIZE"), FComputeShaderUtils::kGolden2DGroupSize);
 		OutEnvironment.SetDefine(TEXT("NUM_CHANNELS"), NumChannels);
+		OutEnvironment.SetDefine(TEXT("KERNEL_SIZE"), 3);
 	}
 };
 
@@ -116,6 +124,8 @@ public:
 		FGlobalShader::ModifyCompilationEnvironment(InParameters, OutEnvironment);
 
 		OutEnvironment.SetDefine(TEXT("TILE_SIZE"), FComputeShaderUtils::kGolden2DGroupSize);
+		OutEnvironment.SetDefine(TEXT("MAX_STEPS"), UE::GeometryMask::Private::MaxSteps);
+		OutEnvironment.SetDefine(TEXT("KERNEL_SIZE"), 3);
 	}
 };
 
@@ -232,11 +242,20 @@ void FGeometryMaskPostProcess_DistanceField::Execute_RenderThread(
 {
 	ensure(IsInRenderingThread());
 
+	uint8 DebugPass = 0; // 0 = no debug
+	
+#if WITH_EDITORONLY_DATA
+	DebugPass = GetDefault<UGeometryMaskSettings>()->DebugDF;
+#endif
+
 	FRDGBuilder GraphBuilder(InRHICmdList);
 	{
 		FIntPoint InputSize = InTexture->GetSizeXY();
 		FVector2f OneOverInputSize = FVector2f::One() / InputSize;
 
+		bool bInputSizeChanged = LastInputSize != InputSize;
+		LastInputSize = InputSize;
+
 		// Scales UV's such that the X axis is always 1.0
 		FVector2f InputHeightOverWidth = FVector2f(1.0f, static_cast<float>(InputSize.Y) / static_cast<float>(InputSize.X));
 
@@ -269,6 +288,9 @@ void FGeometryMaskPostProcess_DistanceField::Execute_RenderThread(
 			SizeFromRadius[ChannelIdx] = Parameters.PerChannelRadius[ChannelIdx];
 		}
 
+		bool bNumActiveChannelsChanged = LastNumActiveChannels != NumActiveChannels;
+		LastNumActiveChannels = NumActiveChannels;
+
 		FIntPoint BufferSize = InputSize;
 		BufferSize.X *= NumActiveChannels;
 		
@@ -324,24 +346,47 @@ void FGeometryMaskPostProcess_DistanceField::Execute_RenderThread(
 
 		FRDGTextureRef InputTexture = InTexture->GetRenderTargetTexture(GraphBuilder);
 		FRDGTextureSRVRef InputTexture_SRV = GraphBuilder.CreateSRV(FRDGTextureSRVDesc(InputTexture));
+		FRDGTextureUAVRef InputTexture_UAV = GraphBuilder.CreateUAV(FRDGTextureUAVDesc(InputTexture));
 
-		FRDGBufferDesc InitOutputBufferDesc =
-			FRDGBufferDesc::CreateBufferDesc(
-				sizeof(FVector4f),
-				BufferLength);
+#if WITH_EDITOR
+		bInputSizeChanged = DebugPass > 0;
+#endif
 
-		FRDGBufferRef InitOutputBuffer = GraphBuilder.CreateBuffer(InitOutputBufferDesc, TEXT("GeometryMaskJFInit.OutputBuffer"));
-		
-		FVector4f* BufferData = GraphBuilder.AllocPODArray<FVector4f>(BufferLength);
-		GraphBuilder.QueueBufferUpload(InitOutputBuffer, BufferData, sizeof(FVector4f) * BufferLength, ERDGInitialDataFlags::NoCopy);
+		FVector4f* BufferData = nullptr;
+		if (bInputSizeChanged || bNumActiveChannelsChanged)
+		{
+			BufferData = GraphBuilder.AllocPODArray<FVector4f>(BufferLength);
+		}
+
+		FRDGBufferRef InitOutputBuffer = nullptr;
+		if (bInputSizeChanged || bNumActiveChannelsChanged)
+		{
+			FRDGBufferDesc InitOutputBufferDesc =
+				FRDGBufferDesc::CreateBufferDesc(
+					sizeof(FVector4f),
+					BufferLength);
+			InitOutputBufferDesc.Usage = EBufferUsageFlags::UnorderedAccess;
+
+			InitOutputBuffer = GraphBuilder.CreateBuffer(InitOutputBufferDesc, TEXT("GeometryMaskJFInit.OutputBuffer"));
+			StoredInitOutputBuffer = GraphBuilder.ConvertToExternalBuffer(InitOutputBuffer);
+
+			GraphBuilder.QueueBufferUpload(InitOutputBuffer, BufferData, sizeof(FVector4f) * BufferLength, ERDGInitialDataFlags::NoCopy);
+		}
+		else
+		{
+			InitOutputBuffer = GraphBuilder.RegisterExternalBuffer(StoredInitOutputBuffer, TEXT("GeometryMaskJFInit.OutputBuffer"));
+		}
 
 		// 1. Init from input binary-ish texture
 		{
 			FRDGBufferUAVRef OutputBuffer_UAV = GraphBuilder.CreateUAV(InitOutputBuffer, EPixelFormat::PF_FloatRGBA);
 
 			{
+				RDG_GPU_STAT_SCOPE(GraphBuilder, GeometryMaskJFInit);
 				RDG_EVENT_SCOPE(GraphBuilder, "GeometryMaskJFInit");
-				
+				TRACE_CPUPROFILER_EVENT_SCOPE(GeometryMaskJFInit);
+				DECLARE_SCOPE_CYCLE_COUNTER(TEXT("FGeometryMaskPostProcess_DistanceField::GeometryMaskJFInit"), STAT_GeometryMask_GeometryMaskJFInit, STATGROUP_GeometryMask);
+
 				FGeometryMaskJFInitCSBase::FParameters* PassParameters = GraphBuilder.AllocParameters<FGeometryMaskJFInitCSBase::FParameters>();
 				{
 					PassParameters->InputDimensions = InputSize;
@@ -353,8 +398,6 @@ void FGeometryMaskPostProcess_DistanceField::Execute_RenderThread(
 					PassParameters->OutputBuffer = OutputBuffer_UAV;
 				}
 
-				ClearUnusedGraphResources(InitComputeShader, PassParameters);
-			
 				FComputeShaderUtils::AddPass(
 					GraphBuilder,
 					RDG_EVENT_NAME("Init"),
@@ -367,40 +410,40 @@ void FGeometryMaskPostProcess_DistanceField::Execute_RenderThread(
 
 		FRDGBufferSRVRef StepOutputBuffer_SRV = nullptr;
 
+		
+#if WITH_EDITOR
+		// DebugPass == 1 - Only sobel
+		if (DebugPass == 1)
+		{
+			StepOutputBuffer_SRV = GraphBuilder.CreateSRV(InitOutputBuffer, EPixelFormat::PF_FloatRGBA);
+		}
+		else
+#endif
 		// 2. JumpFlood steps
 		{
-			FRDGBufferDesc StepIntermediateBufferDescA =
-				FRDGBufferDesc::CreateBufferDesc(
-					sizeof(FVector4f),
-					BufferLength);
-			StepIntermediateBufferDescA.Usage = EBufferUsageFlags::UnorderedAccess;
-
-			FRDGBufferDesc StepIntermediateBufferDescB =
-				FRDGBufferDesc::CreateBufferDesc(
-					sizeof(FVector4f),
-					BufferLength);
-			StepIntermediateBufferDescB.Usage = EBufferUsageFlags::UnorderedAccess;
-
-			FRDGBufferRef StepIntermediateBufferA = GraphBuilder.CreateBuffer(StepIntermediateBufferDescA, TEXT("GeometryMaskJFStep.IntermediateBufferA"));
-			FRDGBufferRef StepIntermediateBufferB = GraphBuilder.CreateBuffer(StepIntermediateBufferDescB, TEXT("GeometryMaskJFStep.IntermediateBufferB"));
+			FRDGBufferRef StepIntermediateBufferB = nullptr;
+			if (bInputSizeChanged || bNumActiveChannelsChanged)
+			{
+				FRDGBufferDesc StepIntermediateBufferDescB =
+					FRDGBufferDesc::CreateBufferDesc(
+						sizeof(FVector4f),
+						BufferLength);
+				StepIntermediateBufferDescB.Usage = EBufferUsageFlags::UnorderedAccess;
 
-			// BufferData = GraphBuilder.AllocPODArray<FVector4f>(BufferLength);
-			GraphBuilder.QueueBufferUpload(StepIntermediateBufferA, BufferData, sizeof(FVector4f) * BufferLength, ERDGInitialDataFlags::NoCopy);
-			
-			// BufferData = GraphBuilder.AllocPODArray<FVector4f>(BufferLength);
-			GraphBuilder.QueueBufferUpload(StepIntermediateBufferB, BufferData, sizeof(FVector4f) * BufferLength, ERDGInitialDataFlags::NoCopy);
+				StepIntermediateBufferB = GraphBuilder.CreateBuffer(StepIntermediateBufferDescB, TEXT("GeometryMaskJFStep.IntermediateBufferB"));
+				StoredStepIntermediateBufferB = GraphBuilder.ConvertToExternalBuffer(StepIntermediateBufferB);
 
-			// Copy init input to intermediate A
+				GraphBuilder.QueueBufferUpload(StepIntermediateBufferB, BufferData, sizeof(FVector4f) * BufferLength, ERDGInitialDataFlags::NoCopy);
+			}
+			else
 			{
-				RDG_EVENT_SCOPE(GraphBuilder, "GeometryMaskJFStep.CopyInitToStep");
-
-				AddCopyBufferPass(GraphBuilder, StepIntermediateBufferA, InitOutputBuffer);
+				StepIntermediateBufferB = GraphBuilder.RegisterExternalBuffer(StoredInitOutputBuffer, TEXT("GeometryMaskJFInit.IntermediateBufferB"));
 			}
-			
-			FRDGBufferUAVRef StepIntermediateBufferA_UAV = GraphBuilder.CreateUAV(StepIntermediateBufferA, EPixelFormat::PF_FloatRGBA);
+
+			FRDGBufferUAVRef StepIntermediateBufferA_UAV = GraphBuilder.CreateUAV(InitOutputBuffer, EPixelFormat::PF_FloatRGBA);
 			FRDGBufferUAVRef StepIntermediateBufferB_UAV = GraphBuilder.CreateUAV(StepIntermediateBufferB, EPixelFormat::PF_FloatRGBA);
 
-			FRDGBufferSRVRef StepIntermediateBufferA_SRV = GraphBuilder.CreateSRV(StepIntermediateBufferA, EPixelFormat::PF_FloatRGBA);
+			FRDGBufferSRVRef StepIntermediateBufferA_SRV = GraphBuilder.CreateSRV(InitOutputBuffer, EPixelFormat::PF_FloatRGBA);
 			FRDGBufferSRVRef StepIntermediateBufferB_SRV = GraphBuilder.CreateSRV(StepIntermediateBufferB, EPixelFormat::PF_FloatRGBA);
 
 			StepOutputBuffer_SRV = StepIntermediateBufferA_SRV;
@@ -410,11 +453,26 @@ void FGeometryMaskPostProcess_DistanceField::Execute_RenderThread(
 				FIntVector4 ChannelMask(ForceInitToZero);
 				int32 MaxStepCount = FMath::Max(FMath::Max(FMath::Max(StepCount.X, StepCount.Y), StepCount.Z), StepCount.W);
 				int32 BufferIdx = 0;
+
+#if WITH_EDITOR
+				if (MaxStepCount > UE::GeometryMask::Private::MaxSteps)
+				{
+					UE_LOG(LogGeometryMask, Warning, TEXT("JumpFlood requires too many steps (%u), maximum is %u."), MaxStepCount, UE::GeometryMask::Private::MaxSteps);
+				}
 				
+				if (DebugPass > 1)
+				{
+					MaxStepCount = FMath::Min(MaxStepCount, DebugPass - 1);
+				}
+#endif
+
 				for (int32 StepIdx = 0; StepIdx < MaxStepCount; ++StepIdx)
 				{
+					RDG_GPU_STAT_SCOPE(GraphBuilder, GeometryMaskJFStep);
 					RDG_EVENT_SCOPE(GraphBuilder, "GeometryMaskJFStep");
-			
+					TRACE_CPUPROFILER_EVENT_SCOPE(GeometryMaskJFStep);
+					DECLARE_SCOPE_CYCLE_COUNTER(TEXT("FGeometryMaskPostProcess_DistanceField::GeometryMaskJFStep"), STAT_GeometryMask_GeometryMaskJFStep, STATGROUP_GeometryMask);
+
 					FGeometryMaskJFStepCSBase::FParameters* PassParameters = GraphBuilder.AllocParameters<FGeometryMaskJFStepCSBase::FParameters>();
 					{
 						PassParameters->InputDimensions = InputSize;
@@ -433,8 +491,6 @@ void FGeometryMaskPostProcess_DistanceField::Execute_RenderThread(
 							: StepIntermediateBufferB_UAV;
 					}
 
-					ClearUnusedGraphResources(StepComputeShader, PassParameters);
-			
 					FComputeShaderUtils::AddPass(
 						GraphBuilder,
 						RDG_EVENT_NAME("Step"),
@@ -454,19 +510,12 @@ void FGeometryMaskPostProcess_DistanceField::Execute_RenderThread(
 
 		// 3. JF to DF
 		{
-			FRDGTextureDesc OutputTextureDesc(
-				FRDGTextureDesc::Create2D(
-					InputTexture->Desc.Extent,
-					InputTexture->Desc.Format,
-					FClearValueBinding::Green,
-					TexCreate_ShaderResource | TexCreate_UAV));
-
-			FRDGTextureRef OutputTexture = GraphBuilder.CreateTexture(OutputTextureDesc, TEXT("GeometryMaskJFtoDF.OutputTexture"));
-			FRDGTextureUAVRef OutputTexture_UAV = GraphBuilder.CreateUAV(FRDGTextureUAVDesc(OutputTexture));
-
 			{
+				RDG_GPU_STAT_SCOPE(GraphBuilder, GeometryMaskJFtoDF);
 				RDG_EVENT_SCOPE(GraphBuilder, "GeometryMaskJFtoDF");
-			
+				TRACE_CPUPROFILER_EVENT_SCOPE(GeometryMaskJFtoDF);
+				DECLARE_SCOPE_CYCLE_COUNTER(TEXT("FGeometryMaskPostProcess_DistanceField::GeometryMaskJFtoDF"), STAT_GeometryMask_GeometryMaskJFtoDF, STATGROUP_GeometryMask);
+
 				FGeometryMaskJFtoDFCSBase::FParameters* PassParameters = GraphBuilder.AllocParameters<FGeometryMaskJFtoDFCSBase::FParameters>();
 				{
 					FVector4f StepDistanceMultipliers;
@@ -484,11 +533,9 @@ void FGeometryMaskPostProcess_DistanceField::Execute_RenderThread(
 					PassParameters->OriginalInputTexture = InputTexture_SRV;
 					PassParameters->OriginalInputSampler = TStaticSamplerState<SF_Bilinear>::GetRHI();
 					PassParameters->InputBuffer = StepOutputBuffer_SRV;
-					PassParameters->OutputTexture = OutputTexture_UAV;
+					PassParameters->OutputTexture = InputTexture_UAV;
 				}
 
-				ClearUnusedGraphResources(OutputComputeShader, PassParameters);
-			
 				FComputeShaderUtils::AddPass(
 					GraphBuilder,
 					RDG_EVENT_NAME("CopyToDF"),
@@ -497,13 +544,6 @@ void FGeometryMaskPostProcess_DistanceField::Execute_RenderThread(
 					PassParameters,
 					NumGroups);
 			}
-
-			// Copy back to RenderTarget
-			{
-				RDG_EVENT_SCOPE(GraphBuilder, "GeometryMaskDF.PostCopy");
-
-				AddCopyTexturePass(GraphBuilder, OutputTexture, InputTexture, FRHICopyTextureInfo());
-			}
 		}
 	}
 
diff --git a/Engine/Plugins/Experimental/GeometryMask/Source/GeometryMask/Private/Shaders/GeometryMaskPostProcess_DistanceField.h b/Engine/Plugins/Experimental/GeometryMask/Source/GeometryMask/Private/Shaders/GeometryMaskPostProcess_DistanceField.h
--- a/Engine/Plugins/Experimental/GeometryMask/Source/GeometryMask/Private/Shaders/GeometryMaskPostProcess_DistanceField.h
+++ b/Engine/Plugins/Experimental/GeometryMask/Source/GeometryMask/Private/Shaders/GeometryMaskPostProcess_DistanceField.h
@@ -6,8 +6,13 @@
 #include "Containers/BitArray.h"
 #include "Containers/ContainerAllocationPolicies.h"
 #include "GeometryMaskPostProcess.h"
+#include "RenderGraphResources.h"
 #include "Templates/SharedPointer.h"
 
+#if WITH_EDITOR
+#include "GeometryMaskSettings.h"
+#endif
+
 class FRenderTarget;
 
 struct FGeometryMaskPostProcessParameters_DistanceField
@@ -32,4 +37,13 @@ public:
 
 protected:
 	virtual void Execute_RenderThread(FRHICommandListImmediate& InRHICmdList, FRenderTarget* InTexture) override;
+
+private:
+	/** Cached to selectively resize resources when needed. */
+	FIntPoint LastInputSize{ForceInitToZero};
+
+	int32 LastNumActiveChannels = 0;
+
+	TRefCountPtr<FRDGPooledBuffer> StoredInitOutputBuffer;
+	TRefCountPtr<FRDGPooledBuffer> StoredStepIntermediateBufferB;
 };
diff --git a/Engine/Plugins/Experimental/GeometryMask/Source/GeometryMask/Public/GeometryMaskCanvas.h b/Engine/Plugins/Experimental/GeometryMask/Source/GeometryMask/Public/GeometryMaskCanvas.h
--- a/Engine/Plugins/Experimental/GeometryMask/Source/GeometryMask/Public/GeometryMaskCanvas.h
+++ b/Engine/Plugins/Experimental/GeometryMask/Source/GeometryMask/Public/GeometryMaskCanvas.h
@@ -48,7 +48,7 @@ public:
 
 	/** Get the underlying render target. */
 	UFUNCTION(BlueprintCallable, BlueprintPure, Category = "Rendering")
-	UTextureRenderTarget2D* GetTexture() const;
+	UCanvasRenderTarget2D* GetTexture() const;
 
 	/** Whether blur is applied or not. */
 	bool IsBlurApplied() const;
diff --git a/Engine/Plugins/Experimental/GeometryMask/Source/GeometryMask/Public/GeometryMaskCanvasActor.h b/Engine/Plugins/Experimental/GeometryMask/Source/GeometryMask/Public/GeometryMaskCanvasActor.h
--- a/Engine/Plugins/Experimental/GeometryMask/Source/GeometryMask/Public/GeometryMaskCanvasActor.h
+++ b/Engine/Plugins/Experimental/GeometryMask/Source/GeometryMask/Public/GeometryMaskCanvasActor.h
@@ -27,7 +27,7 @@ public:
 
 	/** Returns the Canvas Texture. */
 	UFUNCTION(BlueprintCallable, Category = "Rendering")
-	UTextureRenderTarget2D* GetTexture();
+	UCanvasRenderTarget2D* GetTexture();
 
 	virtual void BeginPlay() override;
 
diff --git a/Engine/Plugins/Experimental/GeometryMask/Source/GeometryMask/Public/GeometryMaskCanvasResource.h b/Engine/Plugins/Experimental/GeometryMask/Source/GeometryMask/Public/GeometryMaskCanvasResource.h
--- a/Engine/Plugins/Experimental/GeometryMask/Source/GeometryMask/Public/GeometryMaskCanvasResource.h
+++ b/Engine/Plugins/Experimental/GeometryMask/Source/GeometryMask/Public/GeometryMaskCanvasResource.h
@@ -13,7 +13,7 @@ class FCanvas;
 class FSceneView;
 class IGeometryMaskPostProcess;
 class UCanvas;
-class UTextureRenderTarget2D;
+class UCanvasRenderTarget2D;
 enum class EGeometryMaskColorChannel : uint8;
 
 using FOnGeometryMaskCanvasDraw = TMulticastDelegate<void(FCanvas*)>;
@@ -31,7 +31,8 @@ public:
 	static constexpr int32 MaxTextureSize = 8192;
 	
 	UGeometryMaskCanvasResource();
-
+	virtual ~UGeometryMaskCanvasResource() override;
+	
 	/** Will return the first available color channel without a canvas assigned. EGeometryMaskColorChannel::None if not available. */
 	const EGeometryMaskColorChannel GetNextAvailableColorChannel() const;
 
@@ -48,12 +49,15 @@ public:
 	/** Returns required viewport padding, in pixels - determined by certain effects. */
 	int32 GetViewportPadding() const;
 
-	UTextureRenderTarget2D* GetRenderTargetTexture();
+	UCanvasRenderTarget2D* GetRenderTargetTexture();
 
 	FOnGeometryMaskCanvasDraw& OnDrawToCanvas() { return OnDrawToCanvasDelegate; }
 
 	void UpdateRenderParameters(EGeometryMaskColorChannel InColorChannel, bool bInApplyBlur, double InBlurStrength, bool bInApplyFeather, int32 InOuterFeatherRadius, int32 InInnerFeatherRadius);
 
+	/** Resets the render parameters for the given channel. */
+	void ResetRenderParameters(EGeometryMaskColorChannel InColorChannel);
+
 	/** Updates the canvas, intended to be called every frame. */
 	void Update(UWorld* InWorld, FSceneView& InView);
 
@@ -75,7 +79,7 @@ private:
 	
 	/** The underlying Render Target texture. */
 	UPROPERTY(DuplicateTransient)
-	TObjectPtr<UTextureRenderTarget2D> RenderTargetTexture;
+	TObjectPtr<UCanvasRenderTarget2D> RenderTargetTexture;
 
 	FOnGeometryMaskCanvasDraw OnDrawToCanvasDelegate;
 
diff --git a/Engine/Plugins/Experimental/GeometryMask/Source/GeometryMask/Public/GeometryMaskSubsystem.h b/Engine/Plugins/Experimental/GeometryMask/Source/GeometryMask/Public/GeometryMaskSubsystem.h
--- a/Engine/Plugins/Experimental/GeometryMask/Source/GeometryMask/Public/GeometryMaskSubsystem.h
+++ b/Engine/Plugins/Experimental/GeometryMask/Source/GeometryMask/Public/GeometryMaskSubsystem.h
@@ -30,10 +30,15 @@ public:
 	UFUNCTION(BlueprintCallable, Category = "Canvas")
 	static TArray<FName> GetCanvasNames();
 
+	int32 GetNumActiveCanvasResources() const;
+
 	const TArray<TObjectPtr<UGeometryMaskCanvasResource>>& GetCanvasResources() const;
 
 	void Update(UWorld* InWorld, FSceneViewFamily& InViewFamily);
 
+	/** Toggles if no arg given. */
+	void ToggleUpdate(const TOptional<bool>& bInShouldUpdate = {});
+
 	/** Remove all canvases without any Readers or Writers. Return the number of canvases removed. */
 	int32 RemoveWithoutWriters();
 
@@ -56,6 +61,8 @@ private:
 	FOnGeometryMaskCanvasCreated OnGeometryMaskCanvasCreatedDelegate;
 	FOnGeometryMaskResourceCreated OnGeometryMaskResourceCreatedDelegate;
 
+	std::atomic<bool> bDoUpdates = true;
+
 	UPROPERTY()
 	TMap<FName, TObjectPtr<UGeometryMaskCanvas>> NamedCanvases;
 
diff --git a/Engine/Plugins/Experimental/GeometryMask/Source/GeometryMask/Public/GeometryMaskTypes.h b/Engine/Plugins/Experimental/GeometryMask/Source/GeometryMask/Public/GeometryMaskTypes.h
--- a/Engine/Plugins/Experimental/GeometryMask/Source/GeometryMask/Public/GeometryMaskTypes.h
+++ b/Engine/Plugins/Experimental/GeometryMask/Source/GeometryMask/Public/GeometryMaskTypes.h
@@ -139,13 +139,15 @@ class GEOMETRYMASK_API UGeometryMaskCanvasReferenceComponentBase
 	GENERATED_BODY()
 
 public:
+	virtual ~UGeometryMaskCanvasReferenceComponentBase() override;
+	
 	/** Implement to perform an operation with the provided canvas. */
 	UFUNCTION(BlueprintImplementableEvent, CallInEditor, Category = "Canvas", meta = (DisplayName = "Canvas Changed"))
 	void ReceiveSetCanvas(const UGeometryMaskCanvas* InCanvas);
 
 	/** Returns the Canvas Texture. */
 	UFUNCTION(BlueprintCallable, Category="Rendering")
-	UTextureRenderTarget2D* GetTexture();
+	UCanvasRenderTarget2D* GetTexture();
 
 protected:
 	virtual void BeginPlay() override;
@@ -156,6 +158,8 @@ protected:
 	virtual bool TryResolveCanvas() PURE_VIRTUAL(UGeometryMaskCanvasReferenceComponentBase::TryResolveCanvas, return false; )
 	
 	bool TryResolveNamedCanvas(FName InCanvasName);
+
+	[[maybe_unused]] virtual bool Cleanup();
 	
 protected:
 	FOnGeometryMaskSetCanvasNativeDelegate OnSetCanvasDelegate;
diff --git a/Engine/Plugins/Experimental/GeometryMask/Source/GeometryMask/Public/GeometryMaskWriteComponent.h b/Engine/Plugins/Experimental/GeometryMask/Source/GeometryMask/Public/GeometryMaskWriteComponent.h
--- a/Engine/Plugins/Experimental/GeometryMask/Source/GeometryMask/Public/GeometryMaskWriteComponent.h
+++ b/Engine/Plugins/Experimental/GeometryMask/Source/GeometryMask/Public/GeometryMaskWriteComponent.h
@@ -69,6 +69,7 @@ protected:
 	
 	//~ Begin UGeometryMaskCanvasReferenceComponentBase
     virtual bool TryResolveCanvas() override;
+	virtual bool Cleanup() override;
     //~ End UGeometryMaskCanvasReferenceComponentBase
     	
 	void UpdateCachedData();
@@ -91,7 +92,7 @@ protected:
 	int32 LastPrimitiveComponentCount;
 
 	UPROPERTY(DuplicateTransient)
-	TMap<FName, TWeakObjectPtr<USceneComponent>> CachedComponents;
+	TMap<FName, TWeakObjectPtr<USceneComponent>> CachedComponentsWeak;
 	
 	UPROPERTY(DuplicateTransient)
 	TMap<FName, FGeometryMaskBatchElementData> CachedMeshData;
diff --git a/Engine/Plugins/Experimental/GeometryMask/Source/GeometryMaskEditor/Private/GeometryMaskEditorModule.cpp b/Engine/Plugins/Experimental/GeometryMask/Source/GeometryMaskEditor/Private/GeometryMaskEditorModule.cpp
--- a/Engine/Plugins/Experimental/GeometryMask/Source/GeometryMaskEditor/Private/GeometryMaskEditorModule.cpp
+++ b/Engine/Plugins/Experimental/GeometryMask/Source/GeometryMaskEditor/Private/GeometryMaskEditorModule.cpp
@@ -2,7 +2,10 @@
 
 #include "GeometryMaskEditorModule.h"
 
+#include "Engine/Engine.h"
 #include "Framework/Docking/TabManager.h"
+#include "GeometryMaskEditorLog.h"
+#include "GeometryMaskSubsystem.h"
 #include "Materials/Material.h"
 #include "Styling/SlateIconFinder.h"
 #include "ViewModels/GMECanvasListViewModel.h"
@@ -83,6 +86,20 @@ void FGeometryMaskEditorModule::RegisterConsoleCommands()
 		FConsoleCommandWithArgsDelegate::CreateRaw(this, &FGeometryMaskEditorModule::ExecuteShowVisualizer),
 		ECVF_Default
 	));
+
+	ConsoleCommands.Add(IConsoleManager::Get().RegisterConsoleCommand(
+		TEXT("GeometryMask.Pause"),
+		TEXT("Disable ticking of GeometryMask objects"),
+		FConsoleCommandWithArgsDelegate::CreateRaw(this, &FGeometryMaskEditorModule::ExecutePause),
+		ECVF_Default
+	));
+
+	ConsoleCommands.Add(IConsoleManager::Get().RegisterConsoleCommand(
+		TEXT("GeometryMask.Flush"),
+		TEXT("Flush unused canvas's"),
+		FConsoleCommandWithArgsDelegate::CreateRaw(this, &FGeometryMaskEditorModule::ExecuteFlush),
+		ECVF_Default
+	));
 }
 
 void FGeometryMaskEditorModule::UnregisterConsoleCommands()
@@ -103,6 +120,37 @@ void FGeometryMaskEditorModule::ExecuteShowVisualizer(const TArray<FString>& InA
 	}
 }
 
+void FGeometryMaskEditorModule::ExecutePause(const TArray<FString>& InArgs)
+{
+	if (UGeometryMaskSubsystem* GeometryMaskSubsystem = GEngine->GetEngineSubsystem<UGeometryMaskSubsystem>())
+	{
+		TOptional<bool> bShouldPauseUpdate;
+		if (!InArgs.IsEmpty())
+		{
+			FString ArgStr = InArgs[0];
+			if (ArgStr.IsNumeric())
+			{
+				// Be explicit if arg given
+				GeometryMaskSubsystem->ToggleUpdate(FCString::Atoi64(*ArgStr) == 0);
+				return;
+			}
+		}
+
+		// Otherwise toggle
+		GeometryMaskSubsystem->ToggleUpdate();
+	}
+}
+
+void FGeometryMaskEditorModule::ExecuteFlush(const TArray<FString>& InArgs)
+{
+	if (UGeometryMaskSubsystem* GeometryMaskSubsystem = GEngine->GetEngineSubsystem<UGeometryMaskSubsystem>())
+	{
+		const int32 RemovedCanvasNum = GeometryMaskSubsystem->RemoveWithoutWriters();
+		const int32 ActiveCanvasNum = GeometryMaskSubsystem->GetCanvasNames().Num();
+		UE_LOG(LogGeometryMaskEditor, Display, TEXT("%u canvas's removed because they had no writers - %u canvas's remaining."), RemovedCanvasNum, ActiveCanvasNum);
+	}
+}
+
 #undef LOCTEXT_NAMESPACE
 
 IMPLEMENT_MODULE(FGeometryMaskEditorModule, GeometryMaskEditor)
diff --git a/Engine/Plugins/Experimental/GeometryMask/Source/GeometryMaskEditor/Private/GeometryMaskEditorModule.h b/Engine/Plugins/Experimental/GeometryMask/Source/GeometryMaskEditor/Private/GeometryMaskEditorModule.h
--- a/Engine/Plugins/Experimental/GeometryMask/Source/GeometryMaskEditor/Private/GeometryMaskEditorModule.h
+++ b/Engine/Plugins/Experimental/GeometryMask/Source/GeometryMaskEditor/Private/GeometryMaskEditorModule.h
@@ -29,6 +29,10 @@ private:
 
 	void ExecuteShowVisualizer(const TArray<FString>& InArgs);
 
+	void ExecutePause(const TArray<FString>& InArgs);
+
+	void ExecuteFlush(const TArray<FString>& InArgs);
+
 private:
 	static const FName VisualizerTabId;
 	
diff --git a/Engine/Plugins/Experimental/GeometryMask/Source/GeometryMaskEditor/Private/ViewModels/GMECanvasItemViewModel.cpp b/Engine/Plugins/Experimental/GeometryMask/Source/GeometryMaskEditor/Private/ViewModels/GMECanvasItemViewModel.cpp
--- a/Engine/Plugins/Experimental/GeometryMask/Source/GeometryMaskEditor/Private/ViewModels/GMECanvasItemViewModel.cpp
+++ b/Engine/Plugins/Experimental/GeometryMask/Source/GeometryMaskEditor/Private/ViewModels/GMECanvasItemViewModel.cpp
@@ -2,8 +2,8 @@
 
 #include "GMECanvasItemViewModel.h"
 
+#include "Engine/CanvasRenderTarget2D.h"
 #include "Engine/Texture.h"
-#include "Engine/TextureRenderTarget2D.h"
 #include "GeometryMaskCanvas.h"
 
 TSharedRef<FGMECanvasItemViewModel> FGMECanvasItemViewModel::Create(const TWeakObjectPtr<const UGeometryMaskCanvas>& InCanvas)
@@ -15,9 +15,9 @@ TSharedRef<FGMECanvasItemViewModel> FGMECanvasItemViewModel::Create(const TWeakO
 
 const UTexture* FGMECanvasItemViewModel::GetCanvasTexture() const
 {
-	if (CanvasTexture.IsValid())
+	if (CanvasTextureWeak.IsValid())
 	{
-		return CanvasTexture.Get();
+		return CanvasTextureWeak.Get();
 	}
 	
 	return nullptr;
@@ -41,7 +41,7 @@ FGMECanvasItemViewModel::FGMECanvasItemViewModel(FPrivateToken, const TWeakObjec
 	{
 		CanvasName = Canvas->GetCanvasName();
 		ColorChannel = Canvas->GetColorChannel();
-		CanvasTexture = Canvas->GetTexture();
+		CanvasTextureWeak = Canvas->GetTexture();
 		KnownWriterCount = Canvas->GetWriters().Num();
 	}
 }
diff --git a/Engine/Plugins/Experimental/GeometryMask/Source/GeometryMaskEditor/Private/ViewModels/GMECanvasItemViewModel.h b/Engine/Plugins/Experimental/GeometryMask/Source/GeometryMaskEditor/Private/ViewModels/GMECanvasItemViewModel.h
--- a/Engine/Plugins/Experimental/GeometryMask/Source/GeometryMaskEditor/Private/ViewModels/GMECanvasItemViewModel.h
+++ b/Engine/Plugins/Experimental/GeometryMask/Source/GeometryMaskEditor/Private/ViewModels/GMECanvasItemViewModel.h
@@ -42,7 +42,7 @@ public:
 private:
 	FName CanvasName;
 	EGeometryMaskColorChannel ColorChannel;
-	TWeakObjectPtr<UTexture> CanvasTexture;
+	TWeakObjectPtr<UTexture> CanvasTextureWeak;
 	int32 KnownReaderCount;
 	int32 KnownWriterCount;
 };
diff --git a/Engine/Plugins/Experimental/GeometryMask/Source/GeometryMaskEditor/Private/ViewModels/GMEResourceItemViewModel.cpp b/Engine/Plugins/Experimental/GeometryMask/Source/GeometryMaskEditor/Private/ViewModels/GMEResourceItemViewModel.cpp
--- a/Engine/Plugins/Experimental/GeometryMask/Source/GeometryMaskEditor/Private/ViewModels/GMEResourceItemViewModel.cpp
+++ b/Engine/Plugins/Experimental/GeometryMask/Source/GeometryMaskEditor/Private/ViewModels/GMEResourceItemViewModel.cpp
@@ -2,7 +2,7 @@
 
 #include "GMEResourceItemViewModel.h"
 
-#include "Engine/TextureRenderTarget2D.h"
+#include "Engine/CanvasRenderTarget2D.h"
 #include "GeometryMaskCanvasResource.h"
 
 TSharedRef<FGMEResourceItemViewModel> FGMEResourceItemViewModel::Create(
@@ -15,7 +15,7 @@ TSharedRef<FGMEResourceItemViewModel> FGMEResourceItemViewModel::Create(
 
 float FGMEResourceItemViewModel::GetMemoryUsage() const
 {
-	if (const UTextureRenderTarget2D* Texture = GetResourceTexture())
+	if (const UCanvasRenderTarget2D* Texture = GetResourceTexture())
 	{
 		return (static_cast<float>(Texture->CalcTextureMemorySizeEnum(TMC_ResidentMips)) / (1024.0f * 1024.0f));
 	}
@@ -25,7 +25,7 @@ float FGMEResourceItemViewModel::GetMemoryUsage() const
 
 FIntPoint FGMEResourceItemViewModel::GetDimensions() const
 {
-	if (const UTextureRenderTarget2D* Texture = GetResourceTexture())
+	if (const UCanvasRenderTarget2D* Texture = GetResourceTexture())
 	{
 		return {Texture->SizeX, Texture->SizeY};
 	}
@@ -40,7 +40,7 @@ FGMEResourceItemViewModel::FGMEResourceItemViewModel(
 	if (const UGeometryMaskCanvasResource* Resource = InResource.Get())
 	{
 		UniqueId = Resource->GetUniqueID();
-		ResourceTexture = const_cast<UGeometryMaskCanvasResource*>(Resource)->GetRenderTargetTexture();
+		ResourceTextureWeak = const_cast<UGeometryMaskCanvasResource*>(Resource)->GetRenderTargetTexture();
 	}
 }
 
diff --git a/Engine/Plugins/Experimental/GeometryMask/Source/GeometryMaskEditor/Private/ViewModels/GMEResourceItemViewModel.h b/Engine/Plugins/Experimental/GeometryMask/Source/GeometryMaskEditor/Private/ViewModels/GMEResourceItemViewModel.h
--- a/Engine/Plugins/Experimental/GeometryMask/Source/GeometryMaskEditor/Private/ViewModels/GMEResourceItemViewModel.h
+++ b/Engine/Plugins/Experimental/GeometryMask/Source/GeometryMaskEditor/Private/ViewModels/GMEResourceItemViewModel.h
@@ -9,10 +9,10 @@
 #include "UObject/NameTypes.h"
 #include "UObject/WeakObjectPtr.h"
 
+class UCanvasRenderTarget2D;
 class UGeometryMaskCanvasResource;
 class UGeometryMaskCanvas;
 class UTexture;
-class UTextureRenderTarget2D;
 
 class FGMEResourceItemViewModel
 	: public TSharedFromThis<FGMEResourceItemViewModel>
@@ -30,7 +30,7 @@ public:
 
 	uint32 GetId() const { return UniqueId; }
 
-	const UTextureRenderTarget2D* GetResourceTexture() const { return ResourceTexture.Get(); }
+	const UCanvasRenderTarget2D* GetResourceTexture() const { return ResourceTextureWeak.Get(); }
 	float GetMemoryUsage() const;
 	FIntPoint GetDimensions() const;
 
@@ -43,5 +43,5 @@ public:
 
 private:
 	uint32 UniqueId;
-	TWeakObjectPtr<UTextureRenderTarget2D> ResourceTexture;
+	TWeakObjectPtr<UCanvasRenderTarget2D> ResourceTextureWeak;
 };
diff --git a/Engine/Plugins/Experimental/GeometryMask/Source/GeometryMaskEditor/Private/Widgets/SGMECanvasList.cpp b/Engine/Plugins/Experimental/GeometryMask/Source/GeometryMaskEditor/Private/Widgets/SGMECanvasList.cpp
--- a/Engine/Plugins/Experimental/GeometryMask/Source/GeometryMaskEditor/Private/Widgets/SGMECanvasList.cpp
+++ b/Engine/Plugins/Experimental/GeometryMask/Source/GeometryMaskEditor/Private/Widgets/SGMECanvasList.cpp
@@ -40,7 +40,7 @@ void SGMECanvasList::Construct(const FArguments& InArgs, TSharedRef<FGMECanvasLi
 
 	Refresh();
 
-	TWeakPtr<SGMECanvasList> WeakThisPtr = SharedThis(this);
+	TWeakPtr<SGMECanvasList> ThisWeak = SharedThis(this);
 
 	ChildSlot
 	[
@@ -64,11 +64,11 @@ void SGMECanvasList::Construct(const FArguments& InArgs, TSharedRef<FGMECanvasLi
 				.Padding(Padding)
 				[
 					SNew(STextBlock)
-					.Text_Lambda([WeakThisPtr]()
+					.Text_Lambda([ThisWeak]()
 					{
-						if (WeakThisPtr.IsValid())
+						if (ThisWeak.IsValid())
 						{
-							return FText::AsNumber(WeakThisPtr.Pin()->CanvasItems.Num());
+							return FText::AsNumber(ThisWeak.Pin()->CanvasItems.Num());
 						}
 						return FText::AsNumber(0);
 					})
diff --git a/Engine/Plugins/Experimental/GeometryMask/Source/GeometryMaskEditor/Private/Widgets/SGMEResourceItem.cpp b/Engine/Plugins/Experimental/GeometryMask/Source/GeometryMaskEditor/Private/Widgets/SGMEResourceItem.cpp
--- a/Engine/Plugins/Experimental/GeometryMask/Source/GeometryMaskEditor/Private/Widgets/SGMEResourceItem.cpp
+++ b/Engine/Plugins/Experimental/GeometryMask/Source/GeometryMaskEditor/Private/Widgets/SGMEResourceItem.cpp
@@ -2,7 +2,7 @@
 
 #include "SGMEResourceItem.h"
 
-#include "Engine/TextureRenderTarget2D.h"
+#include "Engine/CanvasRenderTarget2D.h"
 #include "SlateOptMacros.h"
 #include "ViewModels/GMEResourceItemViewModel.h"
 
@@ -16,7 +16,7 @@ void SGMEResourceItem::Construct(
 	ViewModel = InViewModel;
 
 	TextureBrush = MakeShared<FSlateBrush>();
-	TextureBrush->SetResourceObject(const_cast<UTextureRenderTarget2D*>(InViewModel->GetResourceTexture()));
+	TextureBrush->SetResourceObject(const_cast<UCanvasRenderTarget2D*>(InViewModel->GetResourceTexture()));
 	TextureBrush->ImageSize = FVector2D(InViewModel->GetResourceTexture()->SizeX, InViewModel->GetResourceTexture()->SizeY);
 
 	ImageWidget = SNew(SImage)
diff --git a/Engine/Plugins/Experimental/GeometryMask/Source/GeometryMaskEditor/Private/Widgets/SGMEResourceList.cpp b/Engine/Plugins/Experimental/GeometryMask/Source/GeometryMaskEditor/Private/Widgets/SGMEResourceList.cpp
--- a/Engine/Plugins/Experimental/GeometryMask/Source/GeometryMaskEditor/Private/Widgets/SGMEResourceList.cpp
+++ b/Engine/Plugins/Experimental/GeometryMask/Source/GeometryMaskEditor/Private/Widgets/SGMEResourceList.cpp
@@ -40,7 +40,7 @@ void SGMEResourceList::Construct(const FArguments& InArgs, TSharedRef<FGMEResour
 
 	Refresh();
 
-	TWeakPtr<SGMEResourceList> WeakThisPtr = SharedThis(this);
+	TWeakPtr<SGMEResourceList> ThisWeak = SharedThis(this);
 
 	ChildSlot
 	[
@@ -64,11 +64,11 @@ void SGMEResourceList::Construct(const FArguments& InArgs, TSharedRef<FGMEResour
 				.Padding(Padding)
 				[
 					SNew(STextBlock)
-					.Text_Lambda([WeakThisPtr]()
+					.Text_Lambda([ThisWeak]()
 					{
-						if (WeakThisPtr.IsValid())
+						if (ThisWeak.IsValid())
 						{
-							return FText::AsNumber(WeakThisPtr.Pin()->ResourceItems.Num());
+							return FText::AsNumber(ThisWeak.Pin()->ResourceItems.Num());
 						}
 						return FText::AsNumber(0);
 					})
@@ -92,12 +92,12 @@ void SGMEResourceList::Construct(const FArguments& InArgs, TSharedRef<FGMEResour
 				.Padding(Padding)
 				[
 					SNew(STextBlock)
-					.Text_Lambda([WeakThisPtr]()
+					.Text_Lambda([ThisWeak]()
 					{
-						if (WeakThisPtr.IsValid())
+						if (ThisWeak.IsValid())
 						{
 							float Memory = 0.0f;
-							for (TSharedPtr<IGMETreeNodeViewModel>& ChildItem : WeakThisPtr.Pin()->ResourceItems)
+							for (TSharedPtr<IGMETreeNodeViewModel>& ChildItem : ThisWeak.Pin()->ResourceItems)
 							{
 								if (TSharedPtr<FGMEResourceItemViewModel> ResourceItem = StaticCastSharedPtr<FGMEResourceItemViewModel>(ChildItem))
 								{
diff --git a/Engine/Plugins/Experimental/GeometryMask/Source/GeometryMaskEditor/Private/Widgets/SGeometryMaskCanvasPreview.cpp b/Engine/Plugins/Experimental/GeometryMask/Source/GeometryMaskEditor/Private/Widgets/SGeometryMaskCanvasPreview.cpp
--- a/Engine/Plugins/Experimental/GeometryMask/Source/GeometryMaskEditor/Private/Widgets/SGeometryMaskCanvasPreview.cpp
+++ b/Engine/Plugins/Experimental/GeometryMask/Source/GeometryMaskEditor/Private/Widgets/SGeometryMaskCanvasPreview.cpp
@@ -2,9 +2,9 @@
 
 #include "Widgets/SGeometryMaskCanvasPreview.h"
 
+#include "Engine/CanvasRenderTarget2D.h"
 #include "Engine/Engine.h"
 #include "Engine/Texture.h"
-#include "Engine/TextureRenderTarget2D.h"
 #include "GeometryMaskEditorLog.h"
 #include "GeometryMaskSubsystem.h"
 #include "GeometryMaskTypes.h"
@@ -104,7 +104,7 @@ void SGeometryMaskCanvasPreview::Tick(const FGeometry& InAllottedGeometry, const
 {
 	SCompoundWidget::Tick(InAllottedGeometry, InCurrentTime, InDeltaTime);
 
-	UpdateBrush(Canvas.Get(), nullptr);
+	UpdateBrush(CanvasWeak.Get(), nullptr);
 }
 
 const FName SGeometryMaskCanvasPreview::GetCanvasName() const
@@ -184,7 +184,7 @@ void SGeometryMaskCanvasPreview::SetOpacity(const float InValue)
 
 UGeometryMaskCanvas* SGeometryMaskCanvasPreview::GetCanvas() const
 {
-	return Canvas.Get();
+	return CanvasWeak.Get();
 }
 
 FOptionalSize SGeometryMaskCanvasPreview::GetAspectRatio()
@@ -222,11 +222,11 @@ bool SGeometryMaskCanvasPreview::TryResolveCanvas()
 {
 	if (UGeometryMaskSubsystem* Subsystem = GEngine->GetEngineSubsystem<UGeometryMaskSubsystem>())
 	{
-		Canvas = Subsystem->GetNamedCanvas(GetCanvasName());
-		UpdateBrush(Canvas.Get(), nullptr);
+		CanvasWeak = Subsystem->GetNamedCanvas(GetCanvasName());
+		UpdateBrush(CanvasWeak.Get(), nullptr);
 	}
 
-	return Canvas != nullptr;
+	return CanvasWeak != nullptr;
 }
 
 void SGeometryMaskCanvasPreview::UpdateBrush(const UGeometryMaskCanvas* InCanvas, UTexture* InTexture)
diff --git a/Engine/Plugins/Experimental/GeometryMask/Source/GeometryMaskEditor/Public/Widgets/SGeometryMaskCanvasPreview.h b/Engine/Plugins/Experimental/GeometryMask/Source/GeometryMaskEditor/Public/Widgets/SGeometryMaskCanvasPreview.h
--- a/Engine/Plugins/Experimental/GeometryMask/Source/GeometryMaskEditor/Public/Widgets/SGeometryMaskCanvasPreview.h
+++ b/Engine/Plugins/Experimental/GeometryMask/Source/GeometryMaskEditor/Public/Widgets/SGeometryMaskCanvasPreview.h
@@ -86,7 +86,7 @@ private:
 	FSoftObjectPath DefaultTexturePath;
 	TObjectPtr<UTexture> DefaultTexture;
 
-	TWeakObjectPtr<UGeometryMaskCanvas> Canvas;
+	TWeakObjectPtr<UGeometryMaskCanvas> CanvasWeak;
 
 	/** Name of the currently referenced canvas. */
 	TAttribute<FName> CanvasName;
