[GeometryMask] CIS fix

diff --git a/Engine/Plugins/Experimental/GeometryMask/Source/GeometryMask/Private/GeometryMaskReadComponent.cpp b/Engine/Plugins/Experimental/GeometryMask/Source/GeometryMask/Private/GeometryMaskReadComponent.cpp
--- a/Engine/Plugins/Experimental/GeometryMask/Source/GeometryMask/Private/GeometryMaskReadComponent.cpp
+++ b/Engine/Plugins/Experimental/GeometryMask/Source/GeometryMask/Private/GeometryMaskReadComponent.cpp
@@ -2,6 +2,7 @@
 
 #include "GeometryMaskReadComponent.h"
 
+#include "Engine/CanvasRenderTarget2D.h"
 #include "GeometryMaskCanvas.h"
 
 void UGeometryMaskReadComponent::SetParameters(FGeometryMaskReadParameters& InParameters)
diff --git a/Engine/Plugins/Experimental/GeometryMask/Source/GeometryMask/Private/GeometryMaskTypes.cpp b/Engine/Plugins/Experimental/GeometryMask/Source/GeometryMask/Private/GeometryMaskTypes.cpp
--- a/Engine/Plugins/Experimental/GeometryMask/Source/GeometryMask/Private/GeometryMaskTypes.cpp
+++ b/Engine/Plugins/Experimental/GeometryMask/Source/GeometryMask/Private/GeometryMaskTypes.cpp
@@ -2,6 +2,7 @@
 
 #include "GeometryMaskTypes.h"
 
+#include "Engine/CanvasRenderTarget2D.h"
 #include "Engine/Engine.h"
 #include "GeometryMaskSubsystem.h"
 #include "UObject/UObjectThreadContext.h"
diff --git a/Engine/Plugins/Experimental/GeometryMask/Source/GeometryMask/Private/GeometryMaskWriteComponent.cpp b/Engine/Plugins/Experimental/GeometryMask/Source/GeometryMask/Private/GeometryMaskWriteComponent.cpp
--- a/Engine/Plugins/Experimental/GeometryMask/Source/GeometryMask/Private/GeometryMaskWriteComponent.cpp
+++ b/Engine/Plugins/Experimental/GeometryMask/Source/GeometryMask/Private/GeometryMaskWriteComponent.cpp
@@ -8,6 +8,7 @@
 #include "CanvasTypes.h"
 #include "Components/DynamicMeshComponent.h"
 #include "Components/StaticMeshComponent.h"
+#include "Engine/CanvasRenderTarget2D.h"
 #include "Engine/StaticMesh.h"
 #include "GeometryMaskCanvas.h"
 #include "GlobalRenderResources.h"
diff --git a/Engine/Plugins/Experimental/GeometryMask/Source/GeometryMask/Public/GeometryMaskReadComponent.h b/Engine/Plugins/Experimental/GeometryMask/Source/GeometryMask/Public/GeometryMaskReadComponent.h
--- a/Engine/Plugins/Experimental/GeometryMask/Source/GeometryMask/Public/GeometryMaskReadComponent.h
+++ b/Engine/Plugins/Experimental/GeometryMask/Source/GeometryMask/Public/GeometryMaskReadComponent.h
@@ -8,6 +8,7 @@
 
 #include "GeometryMaskReadComponent.generated.h"
 
+class UCanvasRenderTarget2D;
 class UGeometryMaskCanvas;
 
 UCLASS(Blueprintable, BlueprintType, HideCategories=(Activation, Cooking, AssetUserData, Navigation), meta=(BlueprintSpawnableComponent))
diff --git a/Engine/Plugins/Experimental/GeometryMask/Source/GeometryMask/Public/GeometryMaskTypes.h b/Engine/Plugins/Experimental/GeometryMask/Source/GeometryMask/Public/GeometryMaskTypes.h
--- a/Engine/Plugins/Experimental/GeometryMask/Source/GeometryMask/Public/GeometryMaskTypes.h
+++ b/Engine/Plugins/Experimental/GeometryMask/Source/GeometryMask/Public/GeometryMaskTypes.h
@@ -7,6 +7,7 @@
 
 #include "GeometryMaskTypes.generated.h"
 
+class UCanvasRenderTarget2D;
 class UGeometryMaskCanvas;
 
 UENUM(BlueprintType)
diff --git a/Engine/Plugins/Experimental/GeometryMask/Source/GeometryMask/Public/GeometryMaskWriteComponent.h b/Engine/Plugins/Experimental/GeometryMask/Source/GeometryMask/Public/GeometryMaskWriteComponent.h
--- a/Engine/Plugins/Experimental/GeometryMask/Source/GeometryMask/Public/GeometryMaskWriteComponent.h
+++ b/Engine/Plugins/Experimental/GeometryMask/Source/GeometryMask/Public/GeometryMaskWriteComponent.h
@@ -8,6 +8,7 @@
 
 #include "GeometryMaskWriteComponent.generated.h"
 
+class UCanvasRenderTarget2D;
 class UDynamicMeshComponent;
 class UStaticMeshComponent;
 
