[GeometryMask] CIS fix

diff --git a/Engine/Plugins/Experimental/GeometryMask/Source/GeometryMask/Public/GeometryMaskCanvas.h b/Engine/Plugins/Experimental/GeometryMask/Source/GeometryMask/Public/GeometryMaskCanvas.h
--- a/Engine/Plugins/Experimental/GeometryMask/Source/GeometryMask/Public/GeometryMaskCanvas.h
+++ b/Engine/Plugins/Experimental/GeometryMask/Source/GeometryMask/Public/GeometryMaskCanvas.h
@@ -12,6 +12,7 @@
 class FGeometryMaskPostProcess_Blur;
 class IGeometryMaskWriteInterface;
 class UCanvas;
+class UCanvasRenderTarget2D;
 class UGeometryMaskCanvasResource;
 class UGeometryMaskPrimitive;
 class UGeometryMaskSubsystem;
