[GeometryMask]
* Fixed non unity shipping compile error

diff --git a/Engine/Plugins/Experimental/GeometryMask/Source/GeometryMask/Private/GeometryMaskModule.h b/Engine/Plugins/Experimental/GeometryMask/Source/GeometryMask/Private/GeometryMaskModule.h
--- a/Engine/Plugins/Experimental/GeometryMask/Source/GeometryMask/Private/GeometryMaskModule.h
+++ b/Engine/Plugins/Experimental/GeometryMask/Source/GeometryMask/Private/GeometryMaskModule.h
@@ -1,4 +1,4 @@
-﻿// Copyright Epic Games, Inc. All Rights Reserved.
+// Copyright Epic Games, Inc. All Rights Reserved.
 
 #pragma once
 
@@ -6,7 +6,7 @@
 
 #include "Logging/LogMacros.h"
 #include "Modules/ModuleInterface.h"
-#include "Stats/Stats2.h"
+#include "Stats/Stats.h"
 
 DECLARE_LOG_CATEGORY_EXTERN(LogGeometryMask, Log, All);
 DECLARE_STATS_GROUP(TEXT("GeometryMask"), STATGROUP_GeometryMask, STATCAT_Advanced);
