[GeometryMask] Masks are isolated to their outer world

diff --git a/Engine/Plugins/Experimental/Avalanche/Source/AvalancheMask/Private/Mask2D/AvaMask2DBaseModifier.cpp b/Engine/Plugins/Experimental/Avalanche/Source/AvalancheMask/Private/Mask2D/AvaMask2DBaseModifier.cpp
--- a/Engine/Plugins/Experimental/Avalanche/Source/AvalancheMask/Private/Mask2D/AvaMask2DBaseModifier.cpp
+++ b/Engine/Plugins/Experimental/Avalanche/Source/AvalancheMask/Private/Mask2D/AvaMask2DBaseModifier.cpp
@@ -370,7 +370,10 @@ void UAvaMask2DBaseModifier::SaveActorPreState(AActor* InActor, FAvaMask2DActorD
 			, [this, InActor]()
 			{
 				TSharedPtr<IAvaMaskMaterialCollectionHandle> Handle = GetObjectHandleSubsystem()->MakeHandle<IAvaMaskMaterialCollectionHandle>(InActor, UE::AvaMask::Internal::HandleTag);
-				Handle->OnSourceMaterialsChanged().BindUObject(this, &UAvaMask2DBaseModifier::OnMaterialsChanged);
+				if (!Handle->OnSourceMaterialsChanged().IsBoundToObject(this))
+				{
+					Handle->OnSourceMaterialsChanged().BindUObject(this, &UAvaMask2DBaseModifier::OnMaterialsChanged);
+				}
 				return Handle;
 			});
 
@@ -653,9 +656,15 @@ void UAvaMask2DBaseModifier::OnMaskSetCanvas(const UGeometryMaskCanvas* InCanvas
 
 void UAvaMask2DBaseModifier::TryResolveCanvas()
 {
-	if (UGeometryMaskSubsystem* GeometryMaskSubsystem = GEngine->GetEngineSubsystem<UGeometryMaskSubsystem>())
+	const UWorld* World = GetWorld();
+	if (!World)
 	{
-		UGeometryMaskCanvas* Canvas = GeometryMaskSubsystem->GetNamedCanvas(GetChannel());
+		return;
+	}
+	
+	if (UGeometryMaskWorldSubsystem* MaskSubsystem = World->GetSubsystem<UGeometryMaskWorldSubsystem>())
+	{
+		UGeometryMaskCanvas* Canvas = MaskSubsystem->GetNamedCanvas(GetChannel());
 		LastResolvedCanvasName = Canvas->GetCanvasName();
 		CanvasWeak = Canvas;
 	}
@@ -674,13 +683,13 @@ UTexture* UAvaMask2DBaseModifier::TryResolveCanvasTexture(AActor* InActor, FAvaM
 		return CanvasTexture;
 	}
 	
-	if (!GEngine)
+	if (!GetWorld())
 	{
 		return nullptr;
 	}
 
 	// Try get canvas texture if already available
-	UGeometryMaskSubsystem* MaskSubsystem = GEngine->GetEngineSubsystem<UGeometryMaskSubsystem>();
+	UGeometryMaskWorldSubsystem* MaskSubsystem = GetWorld()->GetSubsystem<UGeometryMaskWorldSubsystem>();
 	if (!MaskSubsystem)
 	{
 		return nullptr;
diff --git a/Engine/Plugins/Experimental/Avalanche/Source/AvalancheMask/Private/Mask2D/AvaMask2DReadModifier.cpp b/Engine/Plugins/Experimental/Avalanche/Source/AvalancheMask/Private/Mask2D/AvaMask2DReadModifier.cpp
--- a/Engine/Plugins/Experimental/Avalanche/Source/AvalancheMask/Private/Mask2D/AvaMask2DReadModifier.cpp
+++ b/Engine/Plugins/Experimental/Avalanche/Source/AvalancheMask/Private/Mask2D/AvaMask2DReadModifier.cpp
@@ -109,8 +109,10 @@ bool UAvaMask2DReadModifier::ApplyRead(AActor* InActor, FAvaMask2DActorData& InA
 			ApplyParameters.MaterialParameters.ChannelAsVector = UE::GeometryMask::MaskChannelEnumToVector[ApplyParameters.MaterialParameters.Channel];
 			ApplyParameters.MaterialParameters.bInvert = bInverted;
 
-			const FIntPoint ViewportSize = GeometryCanvas->GetResource()->GetViewportSize();
-			const FVector2f ViewportPadding = FVector2f::One() - FVector2f(ViewportSize + GeometryCanvas->GetResource()->GetViewportPadding()) / ViewportSize;
+			FGeometryMaskDrawingContext DrawingContext(GeometryCanvas->GetCanvasId().World);
+
+			const FIntPoint ViewportSize = GeometryCanvas->GetResource()->GetMaxViewportSize();
+			const FVector2f ViewportPadding = FVector2f::One() - FVector2f(ViewportSize + GeometryCanvas->GetResource()->GetViewportPadding(DrawingContext)) / ViewportSize;
 			ApplyParameters.MaterialParameters.Padding = ViewportPadding;
 
 			ApplyParameters.MaterialParameters.bApplyFeathering = bUseFeathering;
diff --git a/Engine/Plugins/Experimental/Avalanche/Source/AvalancheMaskEditor/Private/AvaMaskEditorMode.cpp b/Engine/Plugins/Experimental/Avalanche/Source/AvalancheMaskEditor/Private/AvaMaskEditorMode.cpp
--- a/Engine/Plugins/Experimental/Avalanche/Source/AvalancheMaskEditor/Private/AvaMaskEditorMode.cpp
+++ b/Engine/Plugins/Experimental/Avalanche/Source/AvalancheMaskEditor/Private/AvaMaskEditorMode.cpp
@@ -13,6 +13,7 @@
 #include "Framework/Application/SlateApplication.h"
 #include "Framework/AvaGizmoComponent.h"
 #include "GeometryMaskSubsystem.h"
+#include "GeometryMaskWorldSubsystem.h"
 #include "IAvaMaskEditor.h"
 #include "IGeometryMaskWriteInterface.h"
 #include "LevelEditor.h"
@@ -268,7 +269,7 @@ void UAvaMaskEditorMode::OnSelectionChanged(const UTypedElementSelectionSet* InS
 {
 	if (InSelectionSet)
 	{
-		PreviewCanvasName = NAME_None;
+		PreviewCanvasId.ResetToNone();
 		PreviewCanvasChannel = EGeometryMaskColorChannel::None;
 		
 		if (const UTypedElementSelectionSet* SelectionSet = WeakActorSelectionSet.Get())
@@ -297,7 +298,7 @@ void UAvaMaskEditorMode::OnSelectionChanged(const UTypedElementSelectionSet* InS
 						}
 					}
 
-					PreviewCanvasName = SelectedMaskCanvas->GetCanvasName();
+					PreviewCanvasId = SelectedMaskCanvas->GetCanvasId();
 					PreviewCanvasChannel = SelectedMaskCanvas->GetColorChannel();
 				}
 			}
@@ -327,7 +328,7 @@ UGeometryMaskCanvas* UAvaMaskEditorMode::GetCanvasReferencedByActor(const AActor
 		CanvasName = ReadComponentParameters.CanvasName;
 	}
 
-	if (UGeometryMaskCanvas* Canvas = GEngine->GetEngineSubsystem<UGeometryMaskSubsystem>()->GetNamedCanvas(CanvasName))
+	if (UGeometryMaskCanvas* Canvas = GetWorld()->GetSubsystem<UGeometryMaskWorldSubsystem>()->GetNamedCanvas(CanvasName))
 	{
 		return Canvas;
 	}
@@ -339,7 +340,7 @@ void UAvaMaskEditorMode::UpdatePreviewWidget()
 {
 	if (CanvasPreviewWidget.IsValid())
 	{
-		CanvasPreviewWidget->SetCanvasName(PreviewCanvasName);
+		CanvasPreviewWidget->SetCanvasId(PreviewCanvasId);
 		CanvasPreviewWidget->SetColorChannel(PreviewCanvasChannel);
 	}
 }
@@ -383,12 +384,12 @@ bool UAvaMaskEditorMode::AddMaskToSelected(const TArray<AActor*>& InMaskingActor
 			MaskWriteModifier->SetUseParentChannel(true);
 
 			// Flush unused canvases, in case temporary actors were created/destroyed
-			if (UGeometryMaskSubsystem* Subsystem = GEngine->GetEngineSubsystem<UGeometryMaskSubsystem>())
+			if (UGeometryMaskWorldSubsystem* Subsystem = GetWorld()->GetSubsystem<UGeometryMaskWorldSubsystem>())
 			{
 				Subsystem->RemoveWithoutWriters();
 			}
 
-			PreviewCanvasName = ChannelName;
+			PreviewCanvasId = FGeometryMaskCanvasId(ParentActor->GetWorld(), ChannelName);
 			PreviewCanvasChannel = ColorChannel;
 		}
 	}
diff --git a/Engine/Plugins/Experimental/Avalanche/Source/AvalancheMaskEditor/Private/AvaMaskEditorMode.h b/Engine/Plugins/Experimental/Avalanche/Source/AvalancheMaskEditor/Private/AvaMaskEditorMode.h
--- a/Engine/Plugins/Experimental/Avalanche/Source/AvalancheMaskEditor/Private/AvaMaskEditorMode.h
+++ b/Engine/Plugins/Experimental/Avalanche/Source/AvalancheMaskEditor/Private/AvaMaskEditorMode.h
@@ -58,7 +58,7 @@ private:
 	TWeakObjectPtr<UGeometryMaskCanvas> SelectedMaskCanvas;
 
 	// @note: these can differ from the selection when interacting
-	FName PreviewCanvasName = NAME_None;
+	FGeometryMaskCanvasId PreviewCanvasId;
 	EGeometryMaskColorChannel PreviewCanvasChannel = EGeometryMaskColorChannel::Red;
 	
 	AActor* GetActorToParentTo() const;
diff --git a/Engine/Plugins/Experimental/GeometryMask/Source/GeometryMask/Private/GeometryMaskCanvas.cpp b/Engine/Plugins/Experimental/GeometryMask/Source/GeometryMask/Private/GeometryMaskCanvas.cpp
--- a/Engine/Plugins/Experimental/GeometryMask/Source/GeometryMask/Private/GeometryMaskCanvas.cpp
+++ b/Engine/Plugins/Experimental/GeometryMask/Source/GeometryMask/Private/GeometryMaskCanvas.cpp
@@ -3,7 +3,6 @@
 #include "GeometryMaskCanvas.h"
 
 #include "Algo/RemoveIf.h"
-#include "CanvasTypes.h"
 #include "Engine/CanvasRenderTarget2D.h"
 #include "Engine/Engine.h"
 #include "Engine/World.h"
@@ -38,7 +37,8 @@ void UGeometryMaskCanvas::PostEditChangeProperty(FPropertyChangedEvent& InProper
 		InnerFeatherRadiusPropertyName
 	};
 
-	if (UpdateRenderParameterProperties.Contains(InPropertyChangedEvent.GetPropertyName()))
+	FName PropertyName = InPropertyChangedEvent.GetPropertyName();
+	if (UpdateRenderParameterProperties.Contains(PropertyName))
 	{
 		UpdateRenderParameters();
 	}
@@ -65,6 +65,14 @@ void UGeometryMaskCanvas::AddWriter(const TScriptInterface<IGeometryMaskWriteInt
 		OnActivated().ExecuteIfBound();
 	}
 
+	if (Writers.ContainsByPredicate([InWriter](const TWeakInterfacePtr<IGeometryMaskWriteInterface>& InExistingWriter)
+	{
+		return InExistingWriter.IsValid() && InExistingWriter.GetObject() == InWriter.GetObject();
+	}))
+	{
+		return;
+	}
+
 	Writers.Add(InWriter.GetObject());
 }
 
@@ -90,7 +98,28 @@ void UGeometryMaskCanvas::RemoveWriter(const TScriptInterface<IGeometryMaskWrite
 	}
 
 	RemoveInvalidWriters();
-	Writers.Remove(InWriter.GetObject());
+
+	Writers.SetNum(Algo::RemoveIf(Writers, [InWriter](const TWeakInterfacePtr<IGeometryMaskWriteInterface>& InExistingWriter)
+	{
+		return InExistingWriter.GetObject() == InWriter.GetObject();
+	}));
+}
+
+int32 UGeometryMaskCanvas::GetNumWriters() const
+{
+	return Writers.Num();
+}
+
+bool UGeometryMaskCanvas::IsDefaultCanvas() const
+{
+	return CanvasId.IsDefault();
+}
+
+void UGeometryMaskCanvas::Free()
+{
+	Writers.Empty();
+	FreeResource();
+	OnDeactivated().ExecuteIfBound();
 }
 
 UCanvasRenderTarget2D* UGeometryMaskCanvas::GetTexture() const
@@ -103,6 +132,11 @@ UCanvasRenderTarget2D* UGeometryMaskCanvas::GetTexture() const
 	return nullptr;
 }
 
+const FGeometryMaskCanvasId& UGeometryMaskCanvas::GetCanvasId() const
+{
+	return CanvasId;
+}
+
 bool UGeometryMaskCanvas::IsBlurApplied() const
 {
 	return bApplyBlur;
@@ -197,7 +231,7 @@ void UGeometryMaskCanvas::RemoveInvalidWriters()
 	}
 }
 
-void UGeometryMaskCanvas::OnDrawToCanvas(FCanvas* InCanvas)
+void UGeometryMaskCanvas::OnDrawToCanvas(const FGeometryMaskDrawingContext& InDrawingContext, FCanvas* InCanvas)
 {
 	// Sort so subtract happens after additive, etc.
 	SortWriters();
@@ -221,6 +255,12 @@ void UGeometryMaskCanvas::UpdateRenderParameters()
 	}
 }
 
+void UGeometryMaskCanvas::Initialize(const UWorld* InWorld, FName InCanvasName)
+{
+	CanvasName = InCanvasName;
+	CanvasId = FGeometryMaskCanvasId(InWorld, CanvasName);
+}
+
 void UGeometryMaskCanvas::Update(
 	UWorld* InWorld,
 	FSceneView& InView)
@@ -232,14 +272,17 @@ void UGeometryMaskCanvas::AssignResource(
 	UGeometryMaskCanvasResource* InResource,
 	const EGeometryMaskColorChannel InColorChannel)
 {
-	// Cleanup existing resource, if assigned
-	FreeResource();
+	// If we're already using this resource (just changing channel), no need to free
+	if (!CanvasResource
+		|| (CanvasResource && CanvasResource != InResource))
+	{
+		FreeResource();
+		CanvasResource = InResource;
+		CanvasResource->OnDrawToCanvas().AddUObject(this, &UGeometryMaskCanvas::OnDrawToCanvas);
+	}
 
-	CanvasResource = InResource;
 	ColorChannel = InColorChannel;
 
-	CanvasResource->OnDrawToCanvas().AddUObject(this, &UGeometryMaskCanvas::OnDrawToCanvas);
-
 	UpdateRenderParameters();
 }
 
@@ -251,7 +294,7 @@ void UGeometryMaskCanvas::FreeResource()
 		ColorChannel = EGeometryMaskColorChannel::None;
 	
 		// Free up the resource
-		CanvasResource->Checkin(CanvasName);
+		CanvasResource->Checkin(CanvasId);
 		CanvasResource = nullptr;
 	}
 }
diff --git a/Engine/Plugins/Experimental/GeometryMask/Source/GeometryMask/Private/GeometryMaskCanvasActor.cpp b/Engine/Plugins/Experimental/GeometryMask/Source/GeometryMask/Private/GeometryMaskCanvasActor.cpp
--- a/Engine/Plugins/Experimental/GeometryMask/Source/GeometryMask/Private/GeometryMaskCanvasActor.cpp
+++ b/Engine/Plugins/Experimental/GeometryMask/Source/GeometryMask/Private/GeometryMaskCanvasActor.cpp
@@ -5,6 +5,7 @@
 #include "Engine/Engine.h"
 #include "GeometryMaskCanvas.h"
 #include "GeometryMaskSubsystem.h"
+#include "GeometryMaskWorldSubsystem.h"
 #include "GeometryMaskWriteComponent.h"
 
 AGeometryMaskCanvasActor::AGeometryMaskCanvasActor(const FObjectInitializer& ObjectInitializer)
@@ -61,7 +62,7 @@ bool AGeometryMaskCanvasActor::TryResolveCanvas()
 {
 	if (!Canvas)
 	{
-		if (UGeometryMaskSubsystem* Subsystem = GEngine->GetEngineSubsystem<UGeometryMaskSubsystem>())
+		if (UGeometryMaskWorldSubsystem* Subsystem = GetWorld()->GetSubsystem<UGeometryMaskWorldSubsystem>())
 		{
 			Canvas = Subsystem->GetNamedCanvas(CanvasName);
 		}
diff --git a/Engine/Plugins/Experimental/GeometryMask/Source/GeometryMask/Private/GeometryMaskCanvasResource.cpp b/Engine/Plugins/Experimental/GeometryMask/Source/GeometryMask/Private/GeometryMaskCanvasResource.cpp
--- a/Engine/Plugins/Experimental/GeometryMask/Source/GeometryMask/Private/GeometryMaskCanvasResource.cpp
+++ b/Engine/Plugins/Experimental/GeometryMask/Source/GeometryMask/Private/GeometryMaskCanvasResource.cpp
@@ -2,12 +2,14 @@
 
 #include "GeometryMaskCanvasResource.h"
 
+#include "Algo/MaxElement.h"
 #include "Engine/Canvas.h"
 #include "Engine/CanvasRenderTarget2D.h"
 #include "Engine/World.h"
 #include "GeometryMaskModule.h"
 #include "GeometryMaskSettings.h"
 #include "GeometryMaskTypes.h"
+#include "GeometryMaskWorldSubsystem.h"
 #include "SceneView.h"
 #include "Shaders/GeometryMaskPostProcess.h"
 #include "Shaders/GeometryMaskPostProcess_Blur.h"
@@ -35,17 +37,16 @@ namespace UE::GeometryMask::Private
 }
 
 UGeometryMaskCanvasResource::UGeometryMaskCanvasResource()
+	: DefaultDrawingContext()
 {
 	// @note: We omit the Alpha channel as it's not uniformly supported
-	DependentCanvasNames = {
-		{ EGeometryMaskColorChannel::Red, NAME_None },
-		{ EGeometryMaskColorChannel::Green, NAME_None },
-		{ EGeometryMaskColorChannel::Blue, NAME_None }
-	};
+	DependentCanvasIds = { { EGeometryMaskColorChannel::Red, FGeometryMaskCanvasId(EForceInit::ForceInit) }, { EGeometryMaskColorChannel::Green, FGeometryMaskCanvasId(EForceInit::ForceInit) }, { EGeometryMaskColorChannel::Blue, FGeometryMaskCanvasId(EForceInit::ForceInit) } };
+
+	UsedChannelMask.Init(false, MaxNumChannels);
 
 	FGeometryMaskPostProcessParameters_Blur PostProcessParameters_Blur;
 	PostProcessParameters_Blur.bPerChannelApplyBlur.SetRange(0, 4, false);
-	PostProcessParameters_Blur.PerChannelBlurStrength = { 16, 16, 16, 16};
+	PostProcessParameters_Blur.PerChannelBlurStrength = { 16, 16, 16, 16 };
 	PostProcess_Blur = MakeShared<FGeometryMaskPostProcess_Blur>(PostProcessParameters_Blur);
 
 	FGeometryMaskPostProcessParameters_DistanceField PostProcessParameters_DistanceField;
@@ -64,7 +65,7 @@ UGeometryMaskCanvasResource::~UGeometryMaskCanvasResource()
 
 const EGeometryMaskColorChannel UGeometryMaskCanvasResource::GetNextAvailableColorChannel() const
 {
-	for (const TPair<EGeometryMaskColorChannel, FName>& ColorChannelCanvas : DependentCanvasNames)
+	for (const TPair<EGeometryMaskColorChannel, FGeometryMaskCanvasId>& ColorChannelCanvas : DependentCanvasIds)
 	{
 		if (ColorChannelCanvas.Value.IsNone())
 		{
@@ -76,80 +77,203 @@ const EGeometryMaskColorChannel UGeometryMaskCanvasResource::GetNextAvailableCol
 }
 
 bool UGeometryMaskCanvasResource::Checkout(
-	const EGeometryMaskColorChannel InColorChannel
-	, const FName InRequestingCanvasName)
+	const EGeometryMaskColorChannel InColorChannel,
+	const FGeometryMaskCanvasId& InRequestingCanvasId)
 {
 	if (!ensure(InColorChannel != EGeometryMaskColorChannel::Num && InColorChannel != EGeometryMaskColorChannel::None))
 	{
+		// Input channel invalid
 		return false;
 	}
 
-	if (!ensure(DependentCanvasNames[InColorChannel].IsNone()))
+	if (!ensure(DependentCanvasIds[InColorChannel].Name.IsNone()))
 	{
+		// Channel already in used
+		return false;
+	}
+
+	DependentCanvasIds[InColorChannel] = InRequestingCanvasId;
+	UsedChannelMask[static_cast<int32>(InColorChannel)] = true;
+	NumUsedChannels++;
+	
+	RemoveInvalidDrawingContexts();
+
+	FGeometryMaskDrawingContext* DrawingContext = GetDrawingContextForCanvas(InRequestingCanvasId);
+	if (!ensure(DrawingContext))
+	{
+		// World is probably invalid
 		return false;
 	}
 
-	DependentCanvasNames[InColorChannel] = InRequestingCanvasName;
 	return true;
 }
 
-bool UGeometryMaskCanvasResource::Checkin(const FName InRequestingCanvasName)
+bool UGeometryMaskCanvasResource::Checkin(const FGeometryMaskCanvasId& InRequestingCanvasId)
 {
-	for (TPair<EGeometryMaskColorChannel, FName>& ColorChannelCanvas : DependentCanvasNames)
+	for (TPair<EGeometryMaskColorChannel, FGeometryMaskCanvasId>& ColorChannelCanvas : DependentCanvasIds)
 	{
-		if (ColorChannelCanvas.Value.IsEqual(InRequestingCanvasName))
+		if (ColorChannelCanvas.Value == InRequestingCanvasId)
 		{
 			// Effectively free this ColorChannel and make available for Checkout
-			ColorChannelCanvas.Value = NAME_None;
+			ColorChannelCanvas.Value.ResetToNone();
 			ResetRenderParameters(ColorChannelCanvas.Key);
+			UsedChannelMask[static_cast<int32>(ColorChannelCanvas.Key)] = false;
+			NumUsedChannels--;
+			
 			return true;
 		}
 	}
 
+	RemoveInvalidDrawingContexts();
+
 	// The requesting canvas name wasn't present in this Resource
 	return false;
 }
 
-void UGeometryMaskCanvasResource::UpdateViewportSize()
+int32 UGeometryMaskCanvasResource::Compact()
 {
-	if (ViewportSize.Size() > 0
-		&& IsValid(RenderTargetTexture)
-		&& !IsUnreachable())
+	int32 NumUsedChannelsAfterCompact = 0;
+	EGeometryMaskColorChannel LastAvailableColorChannel = EGeometryMaskColorChannel::None;
+	for (int32 ChannelIdx = 0; ChannelIdx < MaxNumAvailableChannels; ++ChannelIdx)
 	{
-		const float SizeMultiplier = GetDefault<UGeometryMaskSettings>()->GetDefaultResolutionMultiplier();
-		
-		int32 SizeX = ViewportSize.X * SizeMultiplier;
-		int32 SizeY = ViewportSize.Y * SizeMultiplier;
+		EGeometryMaskColorChannel ColorChannel = static_cast<EGeometryMaskColorChannel>(ChannelIdx);
+		bool bIsChannelUsed = UsedChannelMask[ChannelIdx];
 
-		const int32 Padding = GetViewportPadding();
-		SizeX += Padding * SizeMultiplier;
-		SizeY += Padding * SizeMultiplier;
-
-		if (const float RatioX = static_cast<float>(SizeX) / MaxTextureSize;
-			RatioX > 1.0)
+		if (!bIsChannelUsed && LastAvailableColorChannel == EGeometryMaskColorChannel::None)
 		{
-			// Width too big, cap to max and reduce height proportionally
-			SizeX = MaxTextureSize;
-			SizeY /= RatioX;
+			LastAvailableColorChannel = ColorChannel;
+			continue;
 		}
 
-		if (const float RatioY = static_cast<float>(SizeY) / MaxTextureSize;
-			RatioY > 1.0)
+		// Channel is used, and there was a previous one available to move to
+		if (bIsChannelUsed)
 		{
-			// Height too big, cap to max and reduce width proportionally
-			SizeY = MaxTextureSize;
-			SizeX /= RatioY;
+			// Re-assign to last available channel
+			if (LastAvailableColorChannel != EGeometryMaskColorChannel::None)
+			{
+				FGeometryMaskCanvasId& UsedCanvasId = DependentCanvasIds[ColorChannel];
+				if (UWorld* CanvasWorld = UsedCanvasId.World.ResolveObjectPtr())
+				{
+					if (UGeometryMaskWorldSubsystem* MaskSubsystem = CanvasWorld->GetSubsystem<UGeometryMaskWorldSubsystem>())
+					{
+						if (UGeometryMaskCanvas* UsedCanvas = MaskSubsystem->GetNamedCanvas(UsedCanvasId.Name))
+						{
+							// Re-assign to that last available channel, freeing this one
+							UsedCanvas->AssignResource(this, LastAvailableColorChannel);						
+							++NumUsedChannelsAfterCompact;
+
+							// "Checkin" current channel and make available
+							UsedCanvasId.ResetToNone();
+							ResetRenderParameters(ColorChannel);
+							UsedChannelMask[ChannelIdx] = false;
+							LastAvailableColorChannel = ColorChannel;
+
+							// "Checkout" last channel that this was moved to
+							DependentCanvasIds[LastAvailableColorChannel] = UsedCanvasId;
+							UsedChannelMask[static_cast<int32>(LastAvailableColorChannel)] = true;
+							UsedCanvas->UpdateRenderParameters();
+						}
+					}
+				}
+			}
+			else
+			{
+				++NumUsedChannelsAfterCompact;
+			}
 		}
+	}
+
+	RemoveInvalidDrawingContexts();
+
+	NumUsedChannels = NumUsedChannelsAfterCompact;
+	return MaxNumAvailableChannels - NumUsedChannelsAfterCompact;
+}
 
-		if (RenderTargetTexture->SizeX == SizeX
-			&& RenderTargetTexture->SizeY == SizeY)
+int32 UGeometryMaskCanvasResource::GetNumChannelsUsed() const
+{
+	return NumUsedChannels;
+}
+
+bool UGeometryMaskCanvasResource::IsAnyChannelUsed() const
+{
+	return NumUsedChannels > 0;
+}
+
+TArray<FGeometryMaskCanvasId> UGeometryMaskCanvasResource::GetDependentCanvasIds() const
+{
+	TArray<FGeometryMaskCanvasId> DependentCanvases;
+	DependentCanvases.Reserve(DependentCanvasIds.Num());
+
+	for (const TPair<EGeometryMaskColorChannel, FGeometryMaskCanvasId>& DependentCanvas : DependentCanvasIds)
+	{
+		if (!DependentCanvas.Value.IsNone())
 		{
-			return;
+			DependentCanvases.Emplace(DependentCanvas.Value);
 		}
+	}
+	
+	return DependentCanvases;
+}
+
+void UGeometryMaskCanvasResource::UpdateViewportSize()
+{
+	if (!IsValid(RenderTargetTexture)
+		|| IsUnreachable())
+	{
+		return;
+	}
+	
+	const float SizeMultiplier = GetDefault<UGeometryMaskSettings>()->GetDefaultResolutionMultiplier();
+	
+	// Update based on max required of all drawing contexts
+	FIntPoint NewViewportSize(ForceInitToZero);
+	
+	// Accounts for res multiplier, padding
+	FIntPoint NewScaledViewportSize(ForceInitToZero);
+	
+	for (FGeometryMaskDrawingContext& DrawingContext : DrawingContextCache)
+	{
+		NewViewportSize = NewViewportSize.ComponentMax(DrawingContext.ViewportSize);
+		
+		NewScaledViewportSize = NewScaledViewportSize.ComponentMax(
+			(DrawingContext.ViewportSize + GetViewportPadding(DrawingContext))
+			* SizeMultiplier);
+	}
+
+	if (NewViewportSize.Size() == 0)
+	{
+		return;
+	}
+	
+	MaxViewportSize = NewViewportSize;
+
+	int32 SizeX = NewScaledViewportSize.X;
+	int32 SizeY = NewScaledViewportSize.Y;
+
+	if (const float RatioX = static_cast<float>(SizeX) / MaxTextureSize;
+		RatioX > 1.0)
+	{
+		// Width too big, cap to max and reduce height proportionally
+		SizeX = MaxTextureSize;
+		SizeY /= RatioX;
+	}
+
+	if (const float RatioY = static_cast<float>(SizeY) / MaxTextureSize;
+		RatioY > 1.0)
+	{
+		// Height too big, cap to max and reduce width proportionally
+		SizeY = MaxTextureSize;
+		SizeX /= RatioY;
+	}
 
-		// Update RT size to viewport size
-		RenderTargetTexture->ResizeTarget(SizeX, SizeY);
+	if (RenderTargetTexture->SizeX == SizeX
+		&& RenderTargetTexture->SizeY == SizeY)
+	{
+		return;
 	}
+
+	// Update RT size to viewport size
+	RenderTargetTexture->ResizeTarget(SizeX, SizeY);
 }
 
 void UGeometryMaskCanvasResource::UpdateRenderParameters(
@@ -245,23 +369,28 @@ void UGeometryMaskCanvasResource::ResetRenderParameters(EGeometryMaskColorChanne
 	UpdateRenderParameters(InColorChannel, false, 0.0, false, 0, 0);
 }
 
-void UGeometryMaskCanvasResource::SetViewportSize(const FIntPoint& InViewportSize)
+void UGeometryMaskCanvasResource::SetViewportSize(FGeometryMaskDrawingContext& InDrawingContext, const FIntPoint& InViewportSize)
 {
 	if (InViewportSize.Size() > 0
-		&& ViewportSize != InViewportSize)
+		&& InDrawingContext.ViewportSize != InViewportSize)
 	{
-		ViewportSize = InViewportSize;
-		UpdateViewportSize();		
+		InDrawingContext.ViewportSize = InViewportSize;
+		UpdateViewportSize();
 	}
 }
 
-int32 UGeometryMaskCanvasResource::GetViewportPadding() const
+int32 UGeometryMaskCanvasResource::GetViewportPadding(const FGeometryMaskDrawingContext& InDrawingContext) const
 {
 	int32 MaxFeatherRadius = 0;
 	{
 		const FGeometryMaskPostProcessParameters_DistanceField& DistanceFieldParameters = PostProcess_DistanceField->GetParameters();
 		for (int32 ChannelIdx = 0; ChannelIdx < MaxNumChannels; ++ChannelIdx)
 		{
+			if (!UsedChannelMask[ChannelIdx])
+			{
+				continue;
+			}
+			
 			if (DistanceFieldParameters.bPerChannelCalculateDF[ChannelIdx])
 			{
 				MaxFeatherRadius = FMath::Max(MaxFeatherRadius, DistanceFieldParameters.PerChannelRadius[ChannelIdx]);
@@ -274,6 +403,11 @@ int32 UGeometryMaskCanvasResource::GetViewportPadding() const
 		const FGeometryMaskPostProcessParameters_Blur& BlurParameters = PostProcess_Blur->GetParameters();
 		for (int32 ChannelIdx = 0; ChannelIdx < MaxNumChannels; ++ChannelIdx)
 		{
+			if (!UsedChannelMask[ChannelIdx])
+			{
+				continue;
+			}
+			
 			if (BlurParameters.bPerChannelApplyBlur[ChannelIdx])
 			{
 				MaxBlurRadius = FMath::Max(MaxBlurRadius, UE::GeometryMask::Internal::ComputeEffectiveKernelSize(BlurParameters.PerChannelBlurStrength[ChannelIdx]));
@@ -293,7 +427,7 @@ UCanvasRenderTarget2D* UGeometryMaskCanvasResource::GetRenderTargetTexture()
 		RenderTargetTexture->bForceLinearGamma = false;
 		RenderTargetTexture->bAutoGenerateMips = false;
 		RenderTargetTexture->SetSampleCount(UE::GeometryMask::Private::RenderTargetSampleCount);
-		RenderTargetTexture->InitAutoFormat(ViewportSize.X, ViewportSize.Y);
+		RenderTargetTexture->InitAutoFormat(MaxViewportSize.X, MaxViewportSize.Y);
 	}
 
 	return RenderTargetTexture;
@@ -301,20 +435,24 @@ UCanvasRenderTarget2D* UGeometryMaskCanvasResource::GetRenderTargetTexture()
 
 void UGeometryMaskCanvasResource::Update(
 	UWorld* InWorld,
-	FSceneView& InView)
+	FSceneView& InView,
+	int32 InViewIndex)
 {
-	Draw(InWorld, InView);
+	Draw(InWorld, InView, InViewIndex);
 }
 
-void UGeometryMaskCanvasResource::Draw(UWorld* InWorld, FSceneView& InView)
+void UGeometryMaskCanvasResource::Draw(UWorld* InWorld, FSceneView& InView, int32 InViewIndex)
 {
 	if (!InWorld)
 	{
 		return;
 	}
 
+	UE_LOG(LogGeometryMask, VeryVerbose, TEXT("UGeometryMaskCanvasResource::Update World: %s"), *InWorld->GetName());
+
 	if (UCanvasRenderTarget2D* Texture = GetRenderTargetTexture())
 	{
+		// First time initialization
 		if (!CanvasObject)
 		{
 			CanvasObject = NewObject<UCanvas>(GetTransientPackage());
@@ -331,14 +469,19 @@ void UGeometryMaskCanvasResource::Draw(UWorld* InWorld, FSceneView& InView)
 			CanvasObject->Init(Texture->SizeX, Texture->SizeY, &InView, NewCanvas);
 		}
 
+		FGeometryMaskDrawingContext* DrawingContextPtr = GetDrawingContextForWorld(InWorld, InViewIndex);
+        check(DrawingContextPtr); // Should always be valid - world was already checked
+
+        FGeometryMaskDrawingContext& DrawingContextRef = *DrawingContextPtr;
+
 		// Begin
 		{
-			SetViewportSize(InView.UnconstrainedViewRect.Size());
+			SetViewportSize(DrawingContextRef, InView.UnconstrainedViewRect.Size());
 			FMatrix ProjectionMatrix = InView.ViewMatrices.GetProjectionMatrix();
-			UE::GeometryMask::Private::OverscanProjectionMatrix(ProjectionMatrix, ViewportSize, GetViewportPadding());
+			UE::GeometryMask::Private::OverscanProjectionMatrix(ProjectionMatrix, MaxViewportSize, GetViewportPadding(DrawingContextRef));
 
-			CachedViewProjectionMatrix = InView.ViewMatrices.GetViewMatrix() * ProjectionMatrix;
-			CachedViewProjectionMatrix.M[2][2] = UE_KINDA_SMALL_NUMBER; // Prevents div by zero later
+			DrawingContextRef.ViewProjectionMatrix = InView.ViewMatrices.GetViewMatrix() * ProjectionMatrix;
+			DrawingContextRef.ViewProjectionMatrix.M[2][2] = UE_KINDA_SMALL_NUMBER; // Prevents div by zero later
 
 			UpdateViewportSize();
 			
@@ -357,11 +500,11 @@ void UGeometryMaskCanvasResource::Draw(UWorld* InWorld, FSceneView& InView)
 			const FMatrix CanvasMatrix = CanvasObject->Canvas->GetBottomTransform();
 
 			// Set to World->Viewport transform
-			CanvasObject->Canvas->SetBaseTransform(CachedViewProjectionMatrix);
+			CanvasObject->Canvas->SetBaseTransform(DrawingContextRef.ViewProjectionMatrix);
 			{
 				TRACE_CPUPROFILER_EVENT_SCOPE(UGeometryMaskCanvasResource::Draw::OnDrawToCanvas);
 				
-				OnDrawToCanvas().Broadcast(CanvasObject->Canvas);
+				OnDrawToCanvas().Broadcast(DrawingContextRef, CanvasObject->Canvas);
 			}
 
 			// Restore original transform
@@ -392,3 +535,48 @@ void UGeometryMaskCanvasResource::Draw(UWorld* InWorld, FSceneView& InView)
 		}
 	}
 }
+
+FGeometryMaskDrawingContext* UGeometryMaskCanvasResource::GetDrawingContextForWorld(const UWorld* InWorld, uint8 InSceneViewIndex)
+{
+	if (!ensure(InWorld))
+	{
+		return nullptr;
+	}
+
+	return &DrawingContextCache.FindOrAdd(FGeometryMaskDrawingContext(InWorld, InSceneViewIndex));
+}
+
+FGeometryMaskDrawingContext* UGeometryMaskCanvasResource::GetDrawingContextForCanvas(const FGeometryMaskCanvasId& InCanvasId)
+{
+	if (InCanvasId.IsDefault())
+	{
+		return &DefaultDrawingContext;
+	}
+	
+	return GetDrawingContextForWorld(InCanvasId.World.ResolveObjectPtr(), 0);
+}
+
+FGeometryMaskDrawingContext* UGeometryMaskCanvasResource::GetDrawingContextForChannel(EGeometryMaskColorChannel InColorChannel)
+{
+	return GetDrawingContextForCanvas(DependentCanvasIds[InColorChannel]);
+}
+
+bool UGeometryMaskCanvasResource::RemoveInvalidDrawingContexts()
+{
+	const int32 NumCurrentItems = DrawingContextCache.Num();
+	
+	TSet<FGeometryMaskDrawingContext> ValidDrawingContexts;
+	ValidDrawingContexts.Reserve(NumCurrentItems);
+	
+	for (FGeometryMaskDrawingContext& DrawingContext : DrawingContextCache)
+	{
+		if (DrawingContext.IsValid())
+		{
+			ValidDrawingContexts.Emplace(MoveTemp(DrawingContext));
+		}
+	}
+
+	DrawingContextCache = ValidDrawingContexts;
+
+	return NumCurrentItems != DrawingContextCache.Num();
+}
diff --git a/Engine/Plugins/Experimental/GeometryMask/Source/GeometryMask/Private/GeometryMaskModule.h b/Engine/Plugins/Experimental/GeometryMask/Source/GeometryMask/Private/GeometryMaskModule.h
--- a/Engine/Plugins/Experimental/GeometryMask/Source/GeometryMask/Private/GeometryMaskModule.h
+++ b/Engine/Plugins/Experimental/GeometryMask/Source/GeometryMask/Private/GeometryMaskModule.h
@@ -8,7 +8,7 @@
 #include "Modules/ModuleInterface.h"
 #include "Stats/Stats.h"
 
-DECLARE_LOG_CATEGORY_EXTERN(LogGeometryMask, Log, All);
+DECLARE_LOG_CATEGORY_EXTERN(LogGeometryMask, All, All);
 DECLARE_STATS_GROUP(TEXT("GeometryMask"), STATGROUP_GeometryMask, STATCAT_Advanced);
 
 class FGeometryMaskModule
diff --git a/Engine/Plugins/Experimental/GeometryMask/Source/GeometryMask/Private/GeometryMaskSVE.cpp b/Engine/Plugins/Experimental/GeometryMask/Source/GeometryMask/Private/GeometryMaskSVE.cpp
--- a/Engine/Plugins/Experimental/GeometryMask/Source/GeometryMask/Private/GeometryMaskSVE.cpp
+++ b/Engine/Plugins/Experimental/GeometryMask/Source/GeometryMask/Private/GeometryMaskSVE.cpp
@@ -2,6 +2,7 @@
 
 #include "GeometryMaskSVE.h"
 
+#include "GeometryMaskModule.h"
 #include "GeometryMaskSubsystem.h"
 
 #if WITH_EDITOR
@@ -13,11 +14,12 @@ FGeometryMaskSceneViewExtension::FGeometryMaskSceneViewExtension(
 	UWorld* InWorld)
 	: FWorldSceneViewExtension(AutoRegister, InWorld)
 {
+	UE_LOG(LogGeometryMask, VeryVerbose, TEXT("SVE registered for world: %s"), *InWorld->GetName());
+	
 	GeometryMaskSubsystemWeak = GEngine->GetEngineSubsystem<UGeometryMaskSubsystem>();
 }
 
-void FGeometryMaskSceneViewExtension::BeginRenderViewFamily(
-	FSceneViewFamily& InViewFamily)
+void FGeometryMaskSceneViewExtension::BeginRenderViewFamily(FSceneViewFamily& InViewFamily)
 {
 	if (UGeometryMaskSubsystem* Subsystem = GeometryMaskSubsystemWeak.Get())
 	{
@@ -38,16 +40,20 @@ bool FGeometryMaskSceneViewExtension::IsActiveThisFrame_Internal(const FSceneVie
 	{
 		if (GEditor->IsSimulatingInEditor())
 		{
-			bIsActive = GetWorld()->WorldType == EWorldType::Editor;
+			if (Context.GetWorld()->WorldType == EWorldType::Editor)
+			{
+				return true;
+			}
 		}
-		else if (GEditor->PlayWorld)
+		
+		if (GEditor->PlayWorld)
 		{
-			bIsActive = GetWorld()->WorldType == EWorldType::PIE;
+			bIsActive = Context.GetWorld()->WorldType == EWorldType::PIE;
 		}
 		else
 		{
-			bIsActive = GetWorld()->WorldType == EWorldType::Editor;	
-		}		
+			bIsActive = Context.GetWorld()->WorldType == EWorldType::Editor || GetWorld()->WorldType == EWorldType::Game;
+		}
 	}
 #endif
 
diff --git a/Engine/Plugins/Experimental/GeometryMask/Source/GeometryMask/Private/GeometryMaskSubsystem.cpp b/Engine/Plugins/Experimental/GeometryMask/Source/GeometryMask/Private/GeometryMaskSubsystem.cpp
--- a/Engine/Plugins/Experimental/GeometryMask/Source/GeometryMask/Private/GeometryMaskSubsystem.cpp
+++ b/Engine/Plugins/Experimental/GeometryMask/Source/GeometryMask/Private/GeometryMaskSubsystem.cpp
@@ -10,53 +10,29 @@
 #include "GeometryMaskCanvas.h"
 #include "GeometryMaskCanvasResource.h"
 #include "GeometryMaskModule.h"
+#include "GeometryMaskWorldSubsystem.h"
 #include "SceneView.h"
 #include "UnrealClient.h"
 
-UGeometryMaskCanvas* UGeometryMaskSubsystem::GetNamedCanvas(FName InName)
+UGeometryMaskCanvas* UGeometryMaskSubsystem::GetDefaultCanvas()
 {
-	if (InName.IsNone())
+	if (!DefaultCanvas)
 	{
-		static FName DefaultCanvasName = TEXT("Default");
-		InName = DefaultCanvasName;
+		const FName DefaultCanvasObjectName = FName(TEXT("GeometryMaskCanvas_Default"));
+		DefaultCanvas = NewObject<UGeometryMaskCanvas>(this, DefaultCanvasObjectName);
+		DefaultCanvas->Initialize(nullptr, FGeometryMaskCanvasId::DefaultCanvasName);
+		AssignResourceToCanvas(DefaultCanvas);
 	}
-
-	const FName ObjectName = MakeUniqueObjectName(this, UGeometryMaskCanvas::StaticClass(), FName(FString::Printf(TEXT("GeometryMaskCanvas_%s_"), *InName.ToString())));
-	if (TObjectPtr<UGeometryMaskCanvas>* FoundCanvas = NamedCanvases.Find(InName))
-	{
-		return *FoundCanvas;
-	}
-
-	const TObjectPtr<UGeometryMaskCanvas>& NewCanvas = NamedCanvases.Emplace(InName, NewObject<UGeometryMaskCanvas>(this, ObjectName));
-	NewCanvas->CanvasName = InName;
-	AssignResourceToCanvas(NewCanvas);
-
-	NewCanvas->OnActivated().BindUObject(this, &UGeometryMaskSubsystem::OnCanvasActivated, NewCanvas.Get());
-	NewCanvas->OnDeactivated().BindUObject(this, &UGeometryMaskSubsystem::OnCanvasDeactivated, NewCanvas.Get());
 	
-	OnGeometryMaskCanvasCreatedDelegate.Broadcast(NewCanvas);
-
-	return NewCanvas;
+	return DefaultCanvas;
 }
 
-TArray<FName> UGeometryMaskSubsystem::GetCanvasNames()
+int32 UGeometryMaskSubsystem::GetNumCanvasResources() const
 {
-	if (const UGeometryMaskSubsystem* Subsystem = GEngine->GetEngineSubsystem<UGeometryMaskSubsystem>())
-	{
-		TArray<FName> CanvasNames;
-		Subsystem->NamedCanvases.GenerateKeyArray(CanvasNames);
-		return CanvasNames;
-	}
-	
-	return {};
-}
-
-int32 UGeometryMaskSubsystem::GetNumActiveCanvasResources() const
-{
-	return NamedCanvases.Num();
+	return CanvasResources.Num();
 }
 
-const TArray<TObjectPtr<UGeometryMaskCanvasResource>>& UGeometryMaskSubsystem::GetCanvasResources() const
+const TSet<TObjectPtr<UGeometryMaskCanvasResource>>& UGeometryMaskSubsystem::GetCanvasResources() const
 {
 	return CanvasResources;
 }
@@ -69,26 +45,30 @@ void UGeometryMaskSubsystem::Update(
 	{
 		return;
 	}
+
+	UE_LOG(LogGeometryMask, VeryVerbose, TEXT("UGeometryMaskSubsystem::Update World: %s, Num. Views: %u"), *InWorld->GetName(), InViewFamily.Views.Num());
 	
-	if (UGeometryMaskSubsystem* Subsystem = GEngine->GetEngineSubsystem<UGeometryMaskSubsystem>())
+	if (UGeometryMaskWorldSubsystem* Subsystem = InWorld->GetSubsystem<UGeometryMaskWorldSubsystem>())
 	{
 		DECLARE_SCOPE_CYCLE_COUNTER(TEXT("UGeometryMaskSubsystem::Update"), STAT_GeometryMask_UpdateAll, STATGROUP_GeometryMask);
-	
+
+		int32 ViewIndex = 0;
 		for (const FSceneView*& View : InViewFamily.Views)
 		{
-			FSceneView* MutableView = const_cast<FSceneView*>(View);
+			FSceneView* MutableSceneView = const_cast<FSceneView*>(View);
 
-			// Removes invalid writers
 			for (const TPair<FName, TObjectPtr<UGeometryMaskCanvas>>& NamedCanvas : Subsystem->NamedCanvases)
 			{
-				NamedCanvas.Value->Update(InWorld, *MutableView);			
+				NamedCanvas.Value->Update(InWorld, *MutableSceneView);		
 			}
 
 			// Updates the texture resource
 			for (const TObjectPtr<UGeometryMaskCanvasResource>& Resource : CanvasResources)
 			{
-				Resource->Update(InWorld, *MutableView);
+				Resource->Update(InWorld, *MutableSceneView, ViewIndex);
 			}
+
+			++ViewIndex;
 		}
 	}
 }
@@ -103,32 +83,6 @@ void UGeometryMaskSubsystem::ToggleUpdate(const TOptional<bool>& bInShouldUpdate
 	}
 }
 
-int32 UGeometryMaskSubsystem::RemoveWithoutWriters()
-{
-	int32 NumRemoved = 0;
-	
-	TMap<FName, TObjectPtr<UGeometryMaskCanvas>> UsedCanvases;
-	UsedCanvases.Reserve(NamedCanvases.Num());
-	
-	for (const TPair<FName, TObjectPtr<UGeometryMaskCanvas>>& NamedCanvas : NamedCanvases)
-	{
-		if (NamedCanvas.Key == NAME_None
-			|| !NamedCanvas.Value->GetWriters().IsEmpty())
-		{
-			UsedCanvases.Emplace(NamedCanvas.Key, NamedCanvas.Value);
-		}
-		else
-		{
-			NamedCanvas.Value->FreeResource();
-			++NumRemoved;
-		}
-	}
-	
-	NamedCanvases = UsedCanvases;
-
-	return NumRemoved;
-}
-
 void UGeometryMaskSubsystem::AssignResourceToCanvas(UGeometryMaskCanvas* InCanvas)
 {
 	UGeometryMaskCanvasResource* AvailableResource = nullptr;
@@ -141,45 +95,69 @@ void UGeometryMaskSubsystem::AssignResourceToCanvas(UGeometryMaskCanvas* InCanva
 			AvailableResource = CanvasResource;
 		}
 	}
-	
+
 	// Nothing available, create new resource
 	if (AvailableChannel == EGeometryMaskColorChannel::None)
 	{
-		AvailableResource = CanvasResources.Emplace_GetRef(NewObject<UGeometryMaskCanvasResource>(this));
+		AvailableResource = NewObject<UGeometryMaskCanvasResource>(this);
+		CanvasResources.Emplace(AvailableResource);
 		AvailableChannel = AvailableResource->GetNextAvailableColorChannel();
+		OnGeometryMaskResourceCreatedDelegate.Broadcast(AvailableResource);
 	}
 
-	AvailableResource->Checkout(AvailableChannel, InCanvas->GetCanvasName());
+	AvailableResource->Checkout(AvailableChannel, InCanvas->GetCanvasId());
 	InCanvas->AssignResource(AvailableResource, AvailableChannel);
 }
 
-void UGeometryMaskSubsystem::OnCanvasActivated(UGeometryMaskCanvas* InCanvas)
+void UGeometryMaskSubsystem::CompactResources()
 {
-	if (InCanvas->IsDefaultCanvas())
+	if (CanvasResources.Num() <= 1)
 	{
 		return;
 	}
 
-	// Already has a resource
-	if (InCanvas->GetResource())
+	TSet<TObjectPtr<UGeometryMaskCanvasResource>> UsedResources;
+	UsedResources.Reserve(CanvasResources.Num());
+	
+	for (UGeometryMaskCanvasResource* CanvasResource : CanvasResources)
 	{
-		return;
+		CanvasResource->Compact();
+		if (CanvasResource->IsAnyChannelUsed())
+		{
+			UsedResources.Emplace(CanvasResource);
+		}
+		else
+		{
+			OnGeometryMaskResourceDestroyedDelegate.Broadcast(CanvasResource);
+		}
 	}
 
-	// Provide a new resource for the canvas to write to
-	AssignResourceToCanvas(InCanvas);
+	if (UsedResources.Num() != CanvasResources.Num())
+	{
+		CanvasResources = UsedResources;
+	}
 }
 
-void UGeometryMaskSubsystem::OnCanvasDeactivated(UGeometryMaskCanvas* InCanvas)
+void UGeometryMaskSubsystem::OnWorldDestroyed(UWorld* InWorld)
 {
-	if (InCanvas->IsDefaultCanvas())
+	CompactResources();
+	
+	TSet<TObjectPtr<UGeometryMaskCanvasResource>> UnusedResources;
+	UnusedResources.Reserve(CanvasResources.Num());
+	
+	for (UGeometryMaskCanvasResource* CanvasResource : CanvasResources)
 	{
-		return;
+		if (!CanvasResource->IsAnyChannelUsed())
+		{
+			OnGeometryMaskResourceDestroyedDelegate.Broadcast(CanvasResource);
+			UnusedResources.Emplace(CanvasResource);
+		}
 	}
 
-	if (const UGeometryMaskCanvasResource* CanvasResource = InCanvas->GetResource())
+	if (UnusedResources.IsEmpty())
 	{
-		// Resource assigned, so free it up
-		InCanvas->FreeResource();
+		return;
 	}
+
+	CanvasResources = CanvasResources.Difference(UnusedResources);
 }
diff --git a/Engine/Plugins/Experimental/GeometryMask/Source/GeometryMask/Private/GeometryMaskTypes.cpp b/Engine/Plugins/Experimental/GeometryMask/Source/GeometryMask/Private/GeometryMaskTypes.cpp
--- a/Engine/Plugins/Experimental/GeometryMask/Source/GeometryMask/Private/GeometryMaskTypes.cpp
+++ b/Engine/Plugins/Experimental/GeometryMask/Source/GeometryMask/Private/GeometryMaskTypes.cpp
@@ -5,6 +5,7 @@
 #include "Engine/CanvasRenderTarget2D.h"
 #include "Engine/Engine.h"
 #include "GeometryMaskSubsystem.h"
+#include "GeometryMaskWorldSubsystem.h"
 #include "UObject/UObjectThreadContext.h"
 
 namespace UE::GeometryMask
@@ -39,6 +40,92 @@ namespace UE::GeometryMask
 	}
 }
 
+uint32 GetTypeHash(const FGeometryMaskCanvasId& InCanvasId)
+{
+	return HashCombineFast(
+		HashCombineFast(
+			GetTypeHash(InCanvasId.World),
+			GetTypeHash(InCanvasId.Name)),
+			InCanvasId.SceneViewIndex);
+}
+
+uint32 GetTypeHash(const FGeometryMaskDrawingContext& InUpdateContext)
+{
+	return HashCombineFast(GetTypeHash(InUpdateContext.World), InUpdateContext.SceneViewIndex);
+}
+
+const FGeometryMaskCanvasId& FGeometryMaskCanvasId::None = FGeometryMaskCanvasId(EForceInit::ForceInit);
+const FName FGeometryMaskCanvasId::DefaultCanvasName = TEXT("Default");
+
+FGeometryMaskCanvasId::FGeometryMaskCanvasId(const UWorld* InWorld, const FName InName)
+	: World(InWorld)
+	, Name(InName)
+{
+}
+
+FGeometryMaskCanvasId::FGeometryMaskCanvasId(EForceInit)
+{
+	ResetToNone();
+}
+
+bool FGeometryMaskCanvasId::IsDefault() const
+{
+	return Name.IsEqual(DefaultCanvasName);
+}
+
+bool FGeometryMaskCanvasId::IsNone() const
+{
+	return Name.IsNone();
+}
+
+void FGeometryMaskCanvasId::ResetToNone()
+{
+	World = nullptr;
+	SceneViewIndex = 0;
+	Name = NAME_None;
+}
+
+FString FGeometryMaskCanvasId::ToString() const
+{
+	FString WorldLabel = TEXT("(Transient)");
+	FString WorldTypeLabel = TEXT("");
+	if (const UWorld* ResolvedWorld = World.ResolveObjectPtr())
+	{
+		WorldLabel = ResolvedWorld->GetName();
+		WorldTypeLabel = LexToString(ResolvedWorld->WorldType);
+	}
+
+	return FString::Printf(TEXT("%s(%s).%s"), *WorldLabel, *WorldTypeLabel, *Name.ToString());
+}
+
+FGeometryMaskDrawingContext::FGeometryMaskDrawingContext(TObjectKey<UWorld> InWorld, const uint8 InSceneViewIndex)
+	: World(MoveTemp(InWorld))
+	, SceneViewIndex(InSceneViewIndex)
+	, ViewportSize(ForceInit)
+	, ViewProjectionMatrix(ForceInit)
+{
+}
+
+FGeometryMaskDrawingContext::FGeometryMaskDrawingContext(const UWorld* InWorld, const uint8 InSceneViewIndex)
+	: World(InWorld)
+	, SceneViewIndex(InSceneViewIndex)
+	, ViewportSize(ForceInit)
+	, ViewProjectionMatrix(ForceInit)
+{
+}
+
+FGeometryMaskDrawingContext::FGeometryMaskDrawingContext(EForceInit)
+	: World(nullptr)
+	, ViewportSize(ForceInit)
+	, ViewProjectionMatrix(ForceInit)
+{
+}
+
+bool FGeometryMaskDrawingContext::IsValid() const
+{
+	return World.ResolveObjectPtr() != nullptr;
+}
+
 UGeometryMaskCanvasReferenceComponentBase::~UGeometryMaskCanvasReferenceComponentBase()
 {
 	Cleanup();
@@ -104,7 +191,7 @@ bool UGeometryMaskCanvasReferenceComponentBase::TryResolveNamedCanvas(FName InCa
 		CanvasWeak.Reset();
 	}
 
-	if (UGeometryMaskSubsystem* Subsystem = GEngine->GetEngineSubsystem<UGeometryMaskSubsystem>())
+	if (UGeometryMaskWorldSubsystem* Subsystem = GetWorld()->GetSubsystem<UGeometryMaskWorldSubsystem>())
 	{
 		Canvas = Subsystem->GetNamedCanvas(InCanvasName);
 
diff --git a/Engine/Plugins/Experimental/GeometryMask/Source/GeometryMask/Private/GeometryMaskWorldSubsystem.cpp b/Engine/Plugins/Experimental/GeometryMask/Source/GeometryMask/Private/GeometryMaskWorldSubsystem.cpp
--- a/Engine/Plugins/Experimental/GeometryMask/Source/GeometryMask/Private/GeometryMaskWorldSubsystem.cpp
+++ b/Engine/Plugins/Experimental/GeometryMask/Source/GeometryMask/Private/GeometryMaskWorldSubsystem.cpp
@@ -15,3 +15,132 @@ void UGeometryMaskWorldSubsystem::Initialize(FSubsystemCollectionBase& Collectio
 
 	GeometryMaskSceneViewExtension = FSceneViewExtensions::NewExtension<FGeometryMaskSceneViewExtension>(World);
 }
+
+void UGeometryMaskWorldSubsystem::Deinitialize()
+{
+	Super::Deinitialize();
+
+	for (const TPair<FName, TObjectPtr<UGeometryMaskCanvas>>& NamedCanvas : NamedCanvases)
+	{
+		if (IsValid(NamedCanvas.Value))
+		{
+			NamedCanvas.Value->Free();	
+		}
+	}
+
+	NamedCanvases.Empty();
+
+	if (UGeometryMaskSubsystem* EngineSubsystem = GEngine->GetEngineSubsystem<UGeometryMaskSubsystem>())
+	{
+		EngineSubsystem->OnWorldDestroyed(GetWorld());
+	}
+}
+
+UGeometryMaskCanvas* UGeometryMaskWorldSubsystem::GetNamedCanvas(FName InName)
+{
+	UGeometryMaskSubsystem* EngineSubsystem = GEngine->GetEngineSubsystem<UGeometryMaskSubsystem>();
+	if (!ensureAlwaysMsgf(EngineSubsystem, TEXT("UGeometryMaskSubsystem not resolved.")))
+	{
+		return nullptr;
+	}
+	
+	if (InName.IsNone())
+	{
+		return EngineSubsystem->GetDefaultCanvas();
+	}
+
+	const FName ObjectName = MakeUniqueObjectName(this, UGeometryMaskCanvas::StaticClass(), FName(FString::Printf(TEXT("GeometryMaskCanvas_%s_"), *InName.ToString())));
+	if (TObjectPtr<UGeometryMaskCanvas>* FoundCanvas = NamedCanvases.Find(InName))
+	{
+		return *FoundCanvas;
+	}
+
+	const TObjectPtr<UGeometryMaskCanvas>& NewCanvas = NamedCanvases.Emplace(InName, NewObject<UGeometryMaskCanvas>(this, ObjectName));
+	NewCanvas->Initialize(GetWorld(), InName);
+	EngineSubsystem->AssignResourceToCanvas(NewCanvas);
+
+	NewCanvas->OnActivated().BindUObject(this, &UGeometryMaskWorldSubsystem::OnCanvasActivated, NewCanvas.Get());
+	NewCanvas->OnDeactivated().BindUObject(this, &UGeometryMaskWorldSubsystem::OnCanvasDeactivated, NewCanvas.Get());
+	
+	OnGeometryMaskCanvasCreatedDelegate.Broadcast(NewCanvas);
+
+	return NewCanvas;
+}
+
+TArray<FName> UGeometryMaskWorldSubsystem::GetCanvasNames()
+{
+	TArray<FName> CanvasNames;
+	NamedCanvases.GenerateKeyArray(CanvasNames);
+	return CanvasNames;
+}
+
+int32 UGeometryMaskWorldSubsystem::RemoveWithoutWriters()
+{
+	int32 NumRemoved = 0;
+	
+	TMap<FName, TObjectPtr<UGeometryMaskCanvas>> UsedCanvases;
+	UsedCanvases.Reserve(NamedCanvases.Num());
+
+	for (const TPair<FName, TObjectPtr<UGeometryMaskCanvas>>& NamedCanvas : NamedCanvases)
+	{
+		if (!IsValid(NamedCanvas.Value))
+		{
+			continue;
+		}
+		
+		if (!NamedCanvas.Value->GetWriters().IsEmpty())
+		{
+			UsedCanvases.Emplace(NamedCanvas.Key, NamedCanvas.Value);
+		}
+		else
+		{
+			OnGeometryMaskCanvasDestroyed().Broadcast(NamedCanvas.Value->GetCanvasId());
+			NamedCanvas.Value->FreeResource();
+			++NumRemoved;
+		}
+	}
+	
+	NamedCanvases = UsedCanvases;
+
+	return NumRemoved;
+}
+
+void UGeometryMaskWorldSubsystem::OnCanvasActivated(UGeometryMaskCanvas* InCanvas)
+{
+	check(InCanvas);
+	
+	if (InCanvas->IsDefaultCanvas())
+	{
+		return;
+	}
+
+	// Already has a resource
+	if (InCanvas->GetResource())
+	{
+		return;
+	}
+
+	UGeometryMaskSubsystem* EngineSubsystem = GEngine->GetEngineSubsystem<UGeometryMaskSubsystem>();
+	if (!ensureAlwaysMsgf(EngineSubsystem, TEXT("UGeometryMaskSubsystem not resolved.")))
+	{
+		return;
+	}
+
+	// Provide a new resource for the canvas to write to
+	EngineSubsystem->AssignResourceToCanvas(InCanvas);
+}
+
+void UGeometryMaskWorldSubsystem::OnCanvasDeactivated(UGeometryMaskCanvas* InCanvas)
+{
+	if (!InCanvas || !IsValid(InCanvas) || InCanvas->IsDefaultCanvas())
+	{
+		return;
+	}
+
+	if (const UGeometryMaskCanvasResource* CanvasResource = InCanvas->GetResource())
+	{
+		// Resource assigned, so free it up
+		OnGeometryMaskCanvasDestroyed().Broadcast(InCanvas->GetCanvasId());
+		InCanvas->FreeResource();
+	}
+}
diff --git a/Engine/Plugins/Experimental/GeometryMask/Source/GeometryMask/Private/Shaders/GeometryMaskPostProcess_Blur.cpp b/Engine/Plugins/Experimental/GeometryMask/Source/GeometryMask/Private/Shaders/GeometryMaskPostProcess_Blur.cpp
--- a/Engine/Plugins/Experimental/GeometryMask/Source/GeometryMask/Private/Shaders/GeometryMaskPostProcess_Blur.cpp
+++ b/Engine/Plugins/Experimental/GeometryMask/Source/GeometryMask/Private/Shaders/GeometryMaskPostProcess_Blur.cpp
@@ -117,6 +117,25 @@ namespace UE::GeometryMask
 
 			return KernelSize;
 		}
+		
+		// @see: SBackgroundBlur::ComputeEffectiveKernelSize
+		int32 ComputeEffectiveKernelSize2(double InStrength)
+		{
+			// If the radius isn't set, auto-compute it based on the strength
+			// @note: * 4.0 differs from * 0.3 in the slate code as it seems to produce better auto-computed results
+			int32 KernelSize = FMath::RoundToInt((InStrength * 2) + 1);
+
+			if (KernelSize % 2 == 0)
+			{
+				++KernelSize;
+			}
+
+			static constexpr int32 MinKernelSize = 3;
+			static constexpr int32 MaxKernelSize = 255;
+			KernelSize = FMath::Clamp(KernelSize, MinKernelSize, MaxKernelSize);
+
+			return KernelSize;
+		}
 	}
 }
 
diff --git a/Engine/Plugins/Experimental/GeometryMask/Source/GeometryMask/Public/GeometryMaskCanvas.h b/Engine/Plugins/Experimental/GeometryMask/Source/GeometryMask/Public/GeometryMaskCanvas.h
--- a/Engine/Plugins/Experimental/GeometryMask/Source/GeometryMask/Public/GeometryMaskCanvas.h
+++ b/Engine/Plugins/Experimental/GeometryMask/Source/GeometryMask/Public/GeometryMaskCanvas.h
@@ -48,13 +48,22 @@ public:
 	/** Remove a writer from this canvas. */
 	void RemoveWriter(const TScriptInterface<IGeometryMaskWriteInterface>& InWriter);
 
+	/** Gets the number of writers for the canvas. */
+	int32 GetNumWriters() const;
+
 	/** Returns true if this is the default/blank canvas. */
-	bool IsDefaultCanvas() const { return CanvasName == NAME_None; }
+	bool IsDefaultCanvas() const;
+
+	/** Forcibly free the canvas - removed all writers and frees the resource. */
+	void Free();
 
 	/** Get the underlying render target. */
 	UFUNCTION(BlueprintCallable, BlueprintPure, Category = "Rendering")
 	UCanvasRenderTarget2D* GetTexture() const;
 
+	/** Get the unique canvas id based on the world and canvas name. */
+	const FGeometryMaskCanvasId& GetCanvasId() const;
+
 	/** Whether blur is applied or not. */
 	bool IsBlurApplied() const;
 
@@ -95,6 +104,9 @@ public:
 	/** Called when Writers becomes empty. */
 	FOnGeometryMaskCanvasDeactivated& OnDeactivated() { return OnDeactivatedDelegate; }
 
+	/** Setup with identifying info. */
+	void Initialize(const UWorld* InWorld, FName InCanvasName);
+	
 	/** Updates the canvas, intended to be called every frame. */
 	void Update(UWorld* InWorld, FSceneView& InView);
 	
@@ -117,6 +129,8 @@ private:
 	static const FName InnerFeatherRadiusPropertyName;
 
 private:
+	friend class UGeometryMaskCanvasResource;
+	
 	/** Sorts writers by various criteria for proper rendering order. */
 	void SortWriters();
 	
@@ -124,19 +138,18 @@ private:
 	void RemoveInvalidWriters();
 
 	/** Draws all writers to the canvas. */
-	void OnDrawToCanvas(FCanvas* InCanvas);
+	void OnDrawToCanvas(const FGeometryMaskDrawingContext& InDrawingContext, FCanvas* InCanvas);
 
 	/** Updates relevant shader parameters for the given color channel. */
 	void UpdateRenderParameters();
 	
 private:
-	// Allow the subsystem to set the CanvasName
-	friend class UGeometryMaskSubsystem;
-	friend class UGeometryMaskWorldSubsystem;
-
 	FOnGeometryMaskCanvasActivated OnActivatedDelegate;
 	FOnGeometryMaskCanvasDeactivated OnDeactivatedDelegate;
 
+	UPROPERTY(Transient, DuplicateTransient)
+	FGeometryMaskCanvasId CanvasId;
+
 	/** Uniquely identifies this canvas. */
 	UPROPERTY(VisibleAnywhere, BlueprintReadOnly, Getter = "Auto", Category = "Canvas", meta = (AllowPrivateAccess = "true"))
 	FName CanvasName;
diff --git a/Engine/Plugins/Experimental/GeometryMask/Source/GeometryMask/Public/GeometryMaskCanvasResource.h b/Engine/Plugins/Experimental/GeometryMask/Source/GeometryMask/Public/GeometryMaskCanvasResource.h
--- a/Engine/Plugins/Experimental/GeometryMask/Source/GeometryMask/Public/GeometryMaskCanvasResource.h
+++ b/Engine/Plugins/Experimental/GeometryMask/Source/GeometryMask/Public/GeometryMaskCanvasResource.h
@@ -2,6 +2,7 @@
 
 #pragma once
 
+#include "GeometryMaskTypes.h"
 #include "UObject/Object.h"
 
 #include "GeometryMaskCanvasResource.generated.h"
@@ -16,7 +17,7 @@ class UCanvas;
 class UCanvasRenderTarget2D;
 enum class EGeometryMaskColorChannel : uint8;
 
-using FOnGeometryMaskCanvasDraw = TMulticastDelegate<void(FCanvas*)>;
+using FOnGeometryMaskCanvasDraw = TMulticastDelegate<void(const FGeometryMaskDrawingContext&, FCanvas*)>;
 
 /**  */
 UCLASS()
@@ -28,6 +29,7 @@ class GEOMETRYMASK_API UGeometryMaskCanvasResource
 public:
 	/** (R, G, B, A) */
 	static constexpr int32 MaxNumChannels = 4;
+	static constexpr int32 MaxNumAvailableChannels = 3; // Skip alpha, as it's not universally supported (RT format)
 	static constexpr int32 MaxTextureSize = 8192;
 	
 	UGeometryMaskCanvasResource();
@@ -37,17 +39,27 @@ public:
 	const EGeometryMaskColorChannel GetNextAvailableColorChannel() const;
 
 	/** Requests usage of the given color channel for this resource. Will return true if successful. */
-	bool Checkout(const EGeometryMaskColorChannel InColorChannel, const FName InRequestingCanvasName);
+	bool Checkout(const EGeometryMaskColorChannel InColorChannel, const FGeometryMaskCanvasId& InRequestingCanvasId);
 
 	/** Returns/frees the color channel associated with the given canvas name. Returns true if canvas name found. */
-	bool Checkin(const FName InRequestingCanvasName);
+	bool Checkin(const FGeometryMaskCanvasId& InRequestingCanvasId);
+
+	/** Re-arranges used channels such that they are used sequentially. Returns number of unused channels. */
+	int32 Compact();
+
+	int32 GetNumChannelsUsed() const;
+
+	bool IsAnyChannelUsed() const;
+
+	/** Get the list of dependent canvases, for debug purposes. */
+	TArray<FGeometryMaskCanvasId> GetDependentCanvasIds() const;
 
 	void UpdateViewportSize();
 
-	void SetViewportSize(const FIntPoint& InViewportSize);
+	void SetViewportSize(FGeometryMaskDrawingContext& InDrawingContext, const FIntPoint& InViewportSize);
 
 	/** Returns required viewport padding, in pixels - determined by certain effects. */
-	int32 GetViewportPadding() const;
+	int32 GetViewportPadding(const FGeometryMaskDrawingContext& InDrawingContext) const;
 
 	UCanvasRenderTarget2D* GetRenderTargetTexture();
 
@@ -59,32 +71,44 @@ public:
 	void ResetRenderParameters(EGeometryMaskColorChannel InColorChannel);
 
 	/** Updates the canvas, intended to be called every frame. */
-	void Update(UWorld* InWorld, FSceneView& InView);
+	void Update(UWorld* InWorld, FSceneView& InView, int32 InViewIndex = 0);
 
 private:
 	/** Draws all writers to the canvas. */
-	void Draw(UWorld* InWorld, FSceneView& InView);
+	void Draw(UWorld* InWorld, FSceneView& InView, int32 InViewIndex = 0);
+
+	FGeometryMaskDrawingContext* GetDrawingContextForWorld(const UWorld* InWorld, uint8 InSceneViewIndex);
+	FGeometryMaskDrawingContext* GetDrawingContextForCanvas(const FGeometryMaskCanvasId& InCanvasId);
+	FGeometryMaskDrawingContext* GetDrawingContextForChannel(EGeometryMaskColorChannel InColorChannel);
+
+	[[maybe_unused]] bool RemoveInvalidDrawingContexts();
 
 private:
 	UPROPERTY(NoClear, EditFixedSize, meta = (EditFixedOrder))
-	TMap<EGeometryMaskColorChannel, FName> DependentCanvasNames;
+	TMap<EGeometryMaskColorChannel, FGeometryMaskCanvasId> DependentCanvasIds;
+
+	TBitArray<TInlineAllocator<MaxNumChannels>> UsedChannelMask;
+
+	UPROPERTY()
+	int32 NumUsedChannels = 0;
 
 	/** Underlying UCanvas object. */
 	UPROPERTY(Transient)
 	TObjectPtr<UCanvas> CanvasObject;
 
 	/** The default viewport size to use when it can't be resolved from the actual viewport. */
-	UPROPERTY(EditAnywhere, BlueprintReadWrite, Getter = "Auto", Setter, Category = "Rendering", meta = (AllowPrivateAccess = "true"))
-	FIntPoint ViewportSize = FIntPoint(1920, 1080);
-	
+	UPROPERTY(EditAnywhere, BlueprintReadWrite, Getter = "Auto", Category = "Rendering", meta = (AllowPrivateAccess = "true"))
+	FIntPoint MaxViewportSize = FIntPoint(1920, 1080);
+
 	/** The underlying Render Target texture. */
 	UPROPERTY(DuplicateTransient)
 	TObjectPtr<UCanvasRenderTarget2D> RenderTargetTexture;
 
 	FOnGeometryMaskCanvasDraw OnDrawToCanvasDelegate;
 
-	/** The last resolved ViewProjectionMatrix. */
-	FMatrix CachedViewProjectionMatrix;
+	TSet<FGeometryMaskDrawingContext> DrawingContextCache;
+
+	FGeometryMaskDrawingContext DefaultDrawingContext;
 
 	bool bApplyBlur = false;
 	bool bApplyDF = false;
diff --git a/Engine/Plugins/Experimental/GeometryMask/Source/GeometryMask/Public/GeometryMaskSubsystem.h b/Engine/Plugins/Experimental/GeometryMask/Source/GeometryMask/Public/GeometryMaskSubsystem.h
--- a/Engine/Plugins/Experimental/GeometryMask/Source/GeometryMask/Public/GeometryMaskSubsystem.h
+++ b/Engine/Plugins/Experimental/GeometryMask/Source/GeometryMask/Public/GeometryMaskSubsystem.h
@@ -11,8 +11,8 @@
 
 #include "GeometryMaskSubsystem.generated.h"
 
-using FOnGeometryMaskCanvasCreated = TMulticastDelegate<void(const UGeometryMaskCanvas*)>;
 using FOnGeometryMaskResourceCreated = TMulticastDelegate<void(const UGeometryMaskCanvasResource*)>;
+using FOnGeometryMaskResourceDestroyed = TMulticastDelegate<void(const UGeometryMaskCanvasResource*)>;
 
 /** Maintains the registered named canvases. */
 UCLASS(BlueprintType)
@@ -21,52 +21,49 @@ class GEOMETRYMASK_API UGeometryMaskSubsystem
 {
 	GENERATED_BODY()
 
-public:	
-	/** Retrieves a Canvas, uniquely identified by it's name. */
+public:
+	/** Returns the default, blank canvas. */
 	UFUNCTION(BlueprintCallable, Category = "Canvas")
-	UGeometryMaskCanvas* GetNamedCanvas(FName InName);
+	UGeometryMaskCanvas* GetDefaultCanvas();
 
-	/** Returns all registered canvas names. */
-	UFUNCTION(BlueprintCallable, Category = "Canvas")
-	static TArray<FName> GetCanvasNames();
-
-	int32 GetNumActiveCanvasResources() const;
+	int32 GetNumCanvasResources() const;
 
-	const TArray<TObjectPtr<UGeometryMaskCanvasResource>>& GetCanvasResources() const;
+	const TSet<TObjectPtr<UGeometryMaskCanvasResource>>& GetCanvasResources() const;
 
 	void Update(UWorld* InWorld, FSceneViewFamily& InViewFamily);
 
 	/** Toggles if no arg given. */
 	void ToggleUpdate(const TOptional<bool>& bInShouldUpdate = {});
 
-	/** Remove all canvases without any Readers or Writers. Return the number of canvases removed. */
-	int32 RemoveWithoutWriters();
-
-	/** Called when a new canvas is created due to a unique name being requested. */
-	FOnGeometryMaskCanvasCreated& OnGeometryMaskCanvasCreated() { return OnGeometryMaskCanvasCreatedDelegate; }
-
 	/** Called when a new canvas resource is created. */
 	FOnGeometryMaskResourceCreated& OnGeometryMaskResourceCreated() { return OnGeometryMaskResourceCreatedDelegate; }
 
-private:
+	/** Called when a canvas resource is destroyed. */
+	FOnGeometryMaskResourceDestroyed& OnGeometryMaskResourceDestroyed() { return OnGeometryMaskResourceDestroyedDelegate; }
+
+private:	
 	/** Find and assign the next available resource to the given canvas. */
 	void AssignResourceToCanvas(UGeometryMaskCanvas* InCanvas);
 
-	void OnCanvasActivated(UGeometryMaskCanvas* InCanvas);
-	void OnCanvasDeactivated(UGeometryMaskCanvas* InCanvas);
+	/** Re-arrange used canvases such that they use as few Canvas Resource's as possible.
+	 *  Note that this could cause momentary visual artifacts.
+	 *  This won't check for unused canvas's, only resource channels. */
+	void CompactResources();
+
+	void OnWorldDestroyed(UWorld* InWorld);
 
 private:
 	friend class UGeometryMaskWorldSubsystem;
 
-	FOnGeometryMaskCanvasCreated OnGeometryMaskCanvasCreatedDelegate;
 	FOnGeometryMaskResourceCreated OnGeometryMaskResourceCreatedDelegate;
+	FOnGeometryMaskResourceDestroyed OnGeometryMaskResourceDestroyedDelegate;
 
 	std::atomic<bool> bDoUpdates = true;
 
 	UPROPERTY()
-	TMap<FName, TObjectPtr<UGeometryMaskCanvas>> NamedCanvases;
+	TObjectPtr<UGeometryMaskCanvas> DefaultCanvas;
 
 	/** Pool of GPU/Texture resources used by the canvases. */
 	UPROPERTY(Getter)
-	TArray<TObjectPtr<UGeometryMaskCanvasResource>> CanvasResources;
+	TSet<TObjectPtr<UGeometryMaskCanvasResource>> CanvasResources;
 };
diff --git a/Engine/Plugins/Experimental/GeometryMask/Source/GeometryMask/Public/GeometryMaskTypes.h b/Engine/Plugins/Experimental/GeometryMask/Source/GeometryMask/Public/GeometryMaskTypes.h
--- a/Engine/Plugins/Experimental/GeometryMask/Source/GeometryMask/Public/GeometryMaskTypes.h
+++ b/Engine/Plugins/Experimental/GeometryMask/Source/GeometryMask/Public/GeometryMaskTypes.h
@@ -4,6 +4,7 @@
 
 #include "Components/ActorComponent.h"
 #include "CoreTypes.h"
+#include "UObject/ObjectKey.h"
 
 #include "GeometryMaskTypes.generated.h"
 
@@ -71,6 +72,91 @@ namespace UE::GeometryMask
 	GEOMETRYMASK_API FStringView ChannelToString(EGeometryMaskColorChannel InColorChannel);
 };
 
+USTRUCT()
+struct GEOMETRYMASK_API FGeometryMaskCanvasId
+{
+	GENERATED_BODY()
+
+public:
+	TObjectKey<UWorld> World;
+
+	UPROPERTY()
+	uint8 SceneViewIndex = 0;
+
+	UPROPERTY()
+	FName Name;
+
+public:
+	static const FGeometryMaskCanvasId& None;
+	static const FName DefaultCanvasName;
+
+public:
+	FGeometryMaskCanvasId() = default;
+	explicit FGeometryMaskCanvasId(const UWorld* InWorld, const FName InName);
+	explicit FGeometryMaskCanvasId(EForceInit);
+	
+	bool IsDefault() const;
+	bool IsNone() const;
+	void ResetToNone();
+
+	FString ToString() const;
+
+public:
+	bool operator==(const FGeometryMaskCanvasId& InOther) const
+	{
+		return World == InOther.World && Name.IsEqual(InOther.Name) && SceneViewIndex == InOther.SceneViewIndex;
+	}
+	
+	bool operator!=(const FGeometryMaskCanvasId& InOther) const
+	{
+		return !(*this == InOther);
+	}
+
+	friend uint32 GetTypeHash(const FGeometryMaskCanvasId& InCanvasId);
+};
+
+/** Uniquely identified by world and scene view index. */
+USTRUCT()
+struct GEOMETRYMASK_API FGeometryMaskDrawingContext
+{
+	GENERATED_BODY()
+
+public:
+	TObjectKey<UWorld> World;
+	
+	UPROPERTY()
+	uint8 SceneViewIndex = 0;
+
+	UPROPERTY()
+	FIntPoint ViewportSize;
+
+	/** The last resolved ViewProjectionMatrix. */
+	UPROPERTY()
+	FMatrix ViewProjectionMatrix;
+
+public:
+	FGeometryMaskDrawingContext() = default;
+	explicit FGeometryMaskDrawingContext(TObjectKey<UWorld> InWorld, const uint8 InSceneViewIndex = 0);
+	explicit FGeometryMaskDrawingContext(const UWorld* InWorld, const uint8 InSceneViewIndex = 0);
+	explicit FGeometryMaskDrawingContext(EForceInit);
+	
+	bool IsValid() const;
+
+public:
+	bool operator==(const FGeometryMaskDrawingContext& InOther) const
+	{
+		return World == InOther.World && SceneViewIndex == InOther.SceneViewIndex;
+	}
+
+	bool operator!=(const FGeometryMaskDrawingContext& InOther) const
+	{
+		return !(*this == InOther);
+	}
+
+	friend uint32 GetTypeHash(const FGeometryMaskDrawingContext& InUpdateContext);
+};
+
+
 UENUM(BlueprintType)
 enum class EGeometryMaskCompositeOperation : uint8
 {
diff --git a/Engine/Plugins/Experimental/GeometryMask/Source/GeometryMask/Public/GeometryMaskWorldSubsystem.h b/Engine/Plugins/Experimental/GeometryMask/Source/GeometryMask/Public/GeometryMaskWorldSubsystem.h
--- a/Engine/Plugins/Experimental/GeometryMask/Source/GeometryMask/Public/GeometryMaskWorldSubsystem.h
+++ b/Engine/Plugins/Experimental/GeometryMask/Source/GeometryMask/Public/GeometryMaskWorldSubsystem.h
@@ -2,6 +2,7 @@
 
 #pragma once
 
+#include "GeometryMaskSubsystem.h"
 #include "Subsystems/WorldSubsystem.h"
 #include "UnrealClient.h"
 
@@ -9,18 +10,54 @@
 
 class FGeometryMaskSceneViewExtension;
 
+using FOnGeometryMaskCanvasCreated = TMulticastDelegate<void(const UGeometryMaskCanvas*)>;
+using FOnGeometryMaskCanvasDestroyed = TMulticastDelegate<void(const FGeometryMaskCanvasId&)>;
+
+
+
 /** Updates the canvases. */
 UCLASS()
-class UGeometryMaskWorldSubsystem
+class GEOMETRYMASK_API UGeometryMaskWorldSubsystem
 	: public UWorldSubsystem
 {
 	GENERATED_BODY()
 
+public:
+	/** Retrieves a Canvas, uniquely identified by this world and the canvas name. */
+	UFUNCTION(BlueprintCallable, Category = "Canvas")
+	UGeometryMaskCanvas* GetNamedCanvas(FName InName);
+
+	/** Returns all registered canvas names for this world. */
+	UFUNCTION(BlueprintCallable, Category = "Canvas")
+	TArray<FName> GetCanvasNames();
+	
+	/** Remove all canvases without any Readers or Writers. Return the number of canvases removed. */
+	int32 RemoveWithoutWriters();
+	
+	/** Called when a new canvas is created due to a unique name being requested. */
+	FOnGeometryMaskCanvasCreated& OnGeometryMaskCanvasCreated() { return OnGeometryMaskCanvasCreatedDelegate; }
+
+	/** Called when a canvas is destroyed due to having no registered writers. */
+	FOnGeometryMaskCanvasDestroyed& OnGeometryMaskCanvasDestroyed() { return OnGeometryMaskCanvasDestroyedDelegate; }
+
 protected:
 	// ~Begin USubsystem
 	virtual void Initialize(FSubsystemCollectionBase& Collection) override;
+	virtual void Deinitialize() override;
 	// ~End USubsystem
 
 private:
+	void OnCanvasActivated(UGeometryMaskCanvas* InCanvas);
+	void OnCanvasDeactivated(UGeometryMaskCanvas* InCanvas);
+
+private:
+	friend class UGeometryMaskSubsystem;
+	
 	TSharedPtr<FGeometryMaskSceneViewExtension> GeometryMaskSceneViewExtension;
+
+	FOnGeometryMaskCanvasCreated OnGeometryMaskCanvasCreatedDelegate;
+	FOnGeometryMaskCanvasDestroyed OnGeometryMaskCanvasDestroyedDelegate;
+
+	UPROPERTY()
+	TMap<FName, TObjectPtr<UGeometryMaskCanvas>> NamedCanvases;
 };
diff --git a/Engine/Plugins/Experimental/GeometryMask/Source/GeometryMaskEditor/Private/GeometryMaskEditorModule.cpp b/Engine/Plugins/Experimental/GeometryMask/Source/GeometryMaskEditor/Private/GeometryMaskEditorModule.cpp
--- a/Engine/Plugins/Experimental/GeometryMask/Source/GeometryMaskEditor/Private/GeometryMaskEditorModule.cpp
+++ b/Engine/Plugins/Experimental/GeometryMask/Source/GeometryMaskEditor/Private/GeometryMaskEditorModule.cpp
@@ -6,6 +6,7 @@
 #include "Framework/Docking/TabManager.h"
 #include "GeometryMaskEditorLog.h"
 #include "GeometryMaskSubsystem.h"
+#include "GeometryMaskWorldSubsystem.h"
 #include "Materials/Material.h"
 #include "Styling/SlateIconFinder.h"
 #include "ViewModels/GMECanvasListViewModel.h"
@@ -143,12 +144,20 @@ void FGeometryMaskEditorModule::ExecutePause(const TArray<FString>& InArgs)
 
 void FGeometryMaskEditorModule::ExecuteFlush(const TArray<FString>& InArgs)
 {
-	if (UGeometryMaskSubsystem* GeometryMaskSubsystem = GEngine->GetEngineSubsystem<UGeometryMaskSubsystem>())
+	for (const FWorldContext& WorldContext : GEngine->GetWorldContexts())
 	{
-		const int32 RemovedCanvasNum = GeometryMaskSubsystem->RemoveWithoutWriters();
-		const int32 ActiveCanvasNum = GeometryMaskSubsystem->GetCanvasNames().Num();
-		UE_LOG(LogGeometryMaskEditor, Display, TEXT("%u canvas's removed because they had no writers - %u canvas's remaining."), RemovedCanvasNum, ActiveCanvasNum);
+		if (UWorld* World = WorldContext.World())
+		{
+			if (UGeometryMaskWorldSubsystem* GeometryMaskSubsystem = World->GetSubsystem<UGeometryMaskWorldSubsystem>())
+			{
+				const int32 RemovedCanvasNum = GeometryMaskSubsystem->RemoveWithoutWriters();
+				const int32 ActiveCanvasNum = GeometryMaskSubsystem->GetCanvasNames().Num();
+				UE_LOG(LogGeometryMaskEditor, Display, TEXT("%u canvas's removed because they had no writers - %u canvas's remaining."), RemovedCanvasNum, ActiveCanvasNum);
+			}
+		}
 	}
+	
+
 }
 
 #undef LOCTEXT_NAMESPACE
diff --git a/Engine/Plugins/Experimental/GeometryMask/Source/GeometryMaskEditor/Private/ViewModels/GMECanvasItemViewModel.cpp b/Engine/Plugins/Experimental/GeometryMask/Source/GeometryMaskEditor/Private/ViewModels/GMECanvasItemViewModel.cpp
--- a/Engine/Plugins/Experimental/GeometryMask/Source/GeometryMaskEditor/Private/ViewModels/GMECanvasItemViewModel.cpp
+++ b/Engine/Plugins/Experimental/GeometryMask/Source/GeometryMaskEditor/Private/ViewModels/GMECanvasItemViewModel.cpp
@@ -13,6 +13,24 @@ TSharedRef<FGMECanvasItemViewModel> FGMECanvasItemViewModel::Create(const TWeakO
 	return ViewModel;
 }
 
+FGMECanvasItemViewModel::FGMECanvasItemViewModel(
+	FPrivateToken,
+	const TWeakObjectPtr<const UGeometryMaskCanvas>& InCanvas)
+	: CanvasWeak(InCanvas)
+	, KnownReaderCount(0)
+	, KnownWriterCount(0)
+{
+	if (const UGeometryMaskCanvas* Canvas = InCanvas.Get())
+	{
+		CanvasId = FGeometryMaskCanvasId(Canvas->GetWorld(), Canvas->GetCanvasName());		
+		ColorChannel = Canvas->GetColorChannel();
+		CanvasTextureWeak = Canvas->GetTexture();
+		KnownWriterCount = Canvas->GetWriters().Num();
+
+		UpdateInfoText();
+	}
+}
+
 const UTexture* FGMECanvasItemViewModel::GetCanvasTexture() const
 {
 	if (CanvasTextureWeak.IsValid())
@@ -33,20 +51,59 @@ float FGMECanvasItemViewModel::GetMemoryUsage()
 	return 0.0f;
 }
 
-FGMECanvasItemViewModel::FGMECanvasItemViewModel(FPrivateToken, const TWeakObjectPtr<const UGeometryMaskCanvas>& InCanvas)
-	: KnownReaderCount(0)
-	, KnownWriterCount(0)
+void FGMECanvasItemViewModel::UpdateInfoText()
 {
-	if (const UGeometryMaskCanvas* Canvas = InCanvas.Get())
+	if (const UGeometryMaskCanvas* Canvas = CanvasWeak.Get())
 	{
-		CanvasName = Canvas->GetCanvasName();
-		ColorChannel = Canvas->GetColorChannel();
-		CanvasTextureWeak = Canvas->GetTexture();
-		KnownWriterCount = Canvas->GetWriters().Num();
+		FGeometryMaskCanvasId Id = Canvas->GetCanvasId();
+
+		FString WorldTypeLabel = TEXT("");
+		FString WorldLabel = TEXT("(None)");
+		if (UWorld* CanvasWorld = Id.World.ResolveObjectPtr())
+		{
+			WorldTypeLabel = LexToString(CanvasWorld->WorldType);
+			WorldLabel = CanvasWorld->GetName();
+		}
+		
+		static const FString DefaultCanvasLabel = FGeometryMaskCanvasId::DefaultCanvasName.ToString();
+		FString CanvasLabel = DefaultCanvasLabel;
+		if (!Id.IsDefault())
+		{
+			CanvasLabel = Id.Name.ToString();
+		}
+
+		InfoText = FText::FromString(
+			FString::Printf(TEXT("%-s\n%-12s: %s\n%-12s: %s\n%-12s: %s\n%-12s: %u"),
+				*FString::Printf(TEXT("%s.%s"), *WorldLabel, *CanvasLabel),
+				TEXT("World Type"),	*WorldTypeLabel,
+				TEXT("World"), *WorldLabel,
+				TEXT("Name"), *CanvasLabel,
+				TEXT("Num. Writers"), Canvas->GetNumWriters()));
 	}
 }
 
+bool FGMECanvasItemViewModel::Tick(const float InDeltaSeconds)
+{
+	UpdateInfoText();
+	return true;
+}
+
 bool FGMECanvasItemViewModel::GetChildren(TArray<TSharedPtr<IGMETreeNodeViewModel>>& OutChildren)
 {
 	return false;
 }
+
+const FGeometryMaskCanvasId& FGMECanvasItemViewModel::GetCanvasId() const
+{
+	return CanvasId;
+}
+
+FName FGMECanvasItemViewModel::GetCanvasName() const
+{
+	return CanvasId.Name;
+}
+
+const FText& FGMECanvasItemViewModel::GetCanvasInfo() const
+{
+	return InfoText;
+}
diff --git a/Engine/Plugins/Experimental/GeometryMask/Source/GeometryMaskEditor/Private/ViewModels/GMECanvasItemViewModel.h b/Engine/Plugins/Experimental/GeometryMask/Source/GeometryMaskEditor/Private/ViewModels/GMECanvasItemViewModel.h
--- a/Engine/Plugins/Experimental/GeometryMask/Source/GeometryMaskEditor/Private/ViewModels/GMECanvasItemViewModel.h
+++ b/Engine/Plugins/Experimental/GeometryMask/Source/GeometryMaskEditor/Private/ViewModels/GMECanvasItemViewModel.h
@@ -5,6 +5,7 @@
 #include "Containers/Array.h"
 #include "EditorUndoClient.h"
 #include "GeometryMaskCanvasResource.h"
+#include "GeometryMaskTypes.h"
 #include "GMEViewModelShared.h"
 #include "Templates/SharedPointer.h"
 #include "UObject/NameTypes.h"
@@ -15,33 +16,41 @@ class UTexture;
 
 class FGMECanvasItemViewModel
 	: public TSharedFromThis<FGMECanvasItemViewModel>
+	, public FGMETickableViewModelBase
 	, public FEditorUndoClient
 	, public IGMETreeNodeViewModel
 {
+protected:
+	// Private token only allows members or friends to call MakeShared
+	struct FPrivateToken { explicit FPrivateToken() = default; };
+	
 public:
-	/**  */
 	static TSharedRef<FGMECanvasItemViewModel> Create(const TWeakObjectPtr<const UGeometryMaskCanvas>& InCanvas);
+
+	FGMECanvasItemViewModel(FPrivateToken, const TWeakObjectPtr<const UGeometryMaskCanvas>& InCanvas);
 	virtual ~FGMECanvasItemViewModel() override = default;
+
+	virtual bool Tick(const float InDeltaSeconds) override;
 	
 	// ~Begin IGMETreeNodeViewModel
 	virtual bool GetChildren(TArray<TSharedPtr<IGMETreeNodeViewModel>>& OutChildren) override;
 	// ~End IGMETreeNodeViewModel
 
-	const FName& GetCanvasName() const { return CanvasName; }
+	const FGeometryMaskCanvasId& GetCanvasId() const;
+	FName GetCanvasName() const;
+	const FText& GetCanvasInfo() const;
 	const EGeometryMaskColorChannel GetColorChannel() const { return ColorChannel; }
 	const UTexture* GetCanvasTexture() const;
 	float GetMemoryUsage();
 
 private:
-	// Private token only allows members or friends to call MakeShared
-	struct FPrivateToken { explicit FPrivateToken() = default; };
-
-public:
-	FGMECanvasItemViewModel(FPrivateToken, const TWeakObjectPtr<const UGeometryMaskCanvas>& InCanvas);
+	void UpdateInfoText();
 
 private:
-	FName CanvasName;
+	FGeometryMaskCanvasId CanvasId;
+	FText InfoText;
 	EGeometryMaskColorChannel ColorChannel;
+	TWeakObjectPtr<const UGeometryMaskCanvas> CanvasWeak;
 	TWeakObjectPtr<UTexture> CanvasTextureWeak;
 	int32 KnownReaderCount;
 	int32 KnownWriterCount;
diff --git a/Engine/Plugins/Experimental/GeometryMask/Source/GeometryMaskEditor/Private/ViewModels/GMECanvasListViewModel.cpp b/Engine/Plugins/Experimental/GeometryMask/Source/GeometryMaskEditor/Private/ViewModels/GMECanvasListViewModel.cpp
--- a/Engine/Plugins/Experimental/GeometryMask/Source/GeometryMaskEditor/Private/ViewModels/GMECanvasListViewModel.cpp
+++ b/Engine/Plugins/Experimental/GeometryMask/Source/GeometryMaskEditor/Private/ViewModels/GMECanvasListViewModel.cpp
@@ -4,6 +4,7 @@
 
 #include "Editor.h"
 #include "GeometryMaskSubsystem.h"
+#include "GeometryMaskWorldSubsystem.h"
 #include "GMECanvasItemViewModel.h"
 
 TSharedRef<FGMECanvasListViewModel> FGMECanvasListViewModel::Create()
@@ -16,50 +17,101 @@ TSharedRef<FGMECanvasListViewModel> FGMECanvasListViewModel::Create()
 
 FGMECanvasListViewModel::~FGMECanvasListViewModel()
 {
-	FTSTicker::GetCoreTicker().RemoveTicker(UpdateCheckHandle);
-	
-	if (UGeometryMaskSubsystem* Subsystem = GEngine->GetEngineSubsystem<UGeometryMaskSubsystem>())
+	for (const TObjectKey<UWorld>& WorldKey : LoadedWorlds)
 	{
-		Subsystem->OnGeometryMaskCanvasCreated().Remove(OnCanvasCreatedHandle);
+		if (UWorld* World = WorldKey.ResolveObjectPtr())
+		{
+			if (UGeometryMaskWorldSubsystem* Subsystem = World->GetSubsystem<UGeometryMaskWorldSubsystem>())
+			{
+				if (FDelegateHandle* DelegateHandle = OnCanvasCreatedHandles.Find(World))
+				{
+					Subsystem->OnGeometryMaskCanvasCreated().Remove(*DelegateHandle);
+				}
+
+				if (FDelegateHandle* DelegateHandle = OnCanvasDestroyedHandles.Find(World))
+				{
+					Subsystem->OnGeometryMaskCanvasDestroyed().Remove(*DelegateHandle);
+				}
+			}
+		}
 	}
+
+	OnCanvasCreatedHandles.Reset();
+	OnCanvasDestroyedHandles.Reset();
 	
 	CanvasItems.Reset();
 }
 
-FGMECanvasListViewModel::FGMECanvasListViewModel(FPrivateToken)
+bool FGMECanvasListViewModel::RefreshItems()
 {
-	UpdateCheckHandle = FTSTicker::GetCoreTicker().AddTicker(FTickerDelegate::CreateRaw(this, &FGMECanvasListViewModel::Tick), UpdateCheckInterval);
-}
+	int32 NumCurrentCanvasItems = CanvasItems.Num();
 
-void FGMECanvasListViewModel::Initialize()
-{
-	if (UGeometryMaskSubsystem* Subsystem = GEngine->GetEngineSubsystem<UGeometryMaskSubsystem>())
+	// Add currently registered canvases
 	{
-		RefreshCanvases();
+		CanvasItems.Reset();
 
-		// And listen for new ones
+		for (const TObjectKey<UWorld>& WorldKey : LoadedWorlds)
 		{
-			OnCanvasCreatedHandle = Subsystem->OnGeometryMaskCanvasCreated().AddRaw(this, &FGMECanvasListViewModel::OnCanvasCreated);
+			if (UWorld* World = WorldKey.ResolveObjectPtr())
+			{
+				if (UGeometryMaskWorldSubsystem* Subsystem = World->GetSubsystem<UGeometryMaskWorldSubsystem>())
+				{
+					TArray<FName> CanvasNames = Subsystem->GetCanvasNames();
+					for (const FName CanvasName : CanvasNames)
+					{
+						UGeometryMaskCanvas* Canvas = Subsystem->GetNamedCanvas(CanvasName);
+						CanvasItems.Add(FGMECanvasItemViewModel::Create(Canvas));
+					}
+
+					OnCanvasCreatedHandles.Emplace(World, Subsystem->OnGeometryMaskCanvasCreated().AddRaw(this, &FGMECanvasListViewModel::OnCanvasCreated));
+					OnCanvasDestroyedHandles.Emplace(World, Subsystem->OnGeometryMaskCanvasDestroyed().AddRaw(this, &FGMECanvasListViewModel::OnCanvasDestroyed));
+				}
+			}
 		}
+	}
 
-		OnChanged().Broadcast();
+	return NumCurrentCanvasItems != CanvasItems.Num();
+}
+
+void FGMECanvasListViewModel::OnPostWorldInit(UWorld* InWorld, const UWorld::InitializationValues InWorldValues)
+{
+	FGMEListViewModelBase::OnPostWorldInit(InWorld, InWorldValues);
+
+	// Listen for new canvases, and destroyed ones
+	{
+		for (const TObjectKey<UWorld>& WorldKey : LoadedWorlds)
+		{
+			if (UWorld* World = WorldKey.ResolveObjectPtr())
+			{
+				if (UGeometryMaskWorldSubsystem* Subsystem = World->GetSubsystem<UGeometryMaskWorldSubsystem>())
+				{
+					OnCanvasCreatedHandles.Emplace(World, Subsystem->OnGeometryMaskCanvasCreated().AddRaw(this, &FGMECanvasListViewModel::OnCanvasCreated));
+					OnCanvasDestroyedHandles.Emplace(World, Subsystem->OnGeometryMaskCanvasDestroyed().AddRaw(this, &FGMECanvasListViewModel::OnCanvasDestroyed));
+				}
+			}
+		}
 	}
 }
 
-void FGMECanvasListViewModel::RefreshCanvases()
+void FGMECanvasListViewModel::OnPreWorldDestroyed(UWorld* InWorld)
 {
-	if (UGeometryMaskSubsystem* Subsystem = GEngine->GetEngineSubsystem<UGeometryMaskSubsystem>())
+	FGMEListViewModelBase::OnPreWorldDestroyed(InWorld);
+
+	for (const TObjectKey<UWorld>& WorldKey : LoadedWorlds)
 	{
-		// Add currently registered canvases
+		if (UWorld* World = WorldKey.ResolveObjectPtr())
 		{
-			TArray<FName> CanvasNames = Subsystem->GetCanvasNames();
-			CanvasItems.Reset();
-			CanvasItems.Reserve(CanvasNames.Num());
-		
-			for (const FName CanvasName : CanvasNames)
+			if (UGeometryMaskWorldSubsystem* Subsystem = World->GetSubsystem<UGeometryMaskWorldSubsystem>())
 			{
-				UGeometryMaskCanvas* Canvas = Subsystem->GetNamedCanvas(CanvasName);
-				CanvasItems.Add(FGMECanvasItemViewModel::Create(Canvas));
+				if (FDelegateHandle* DelegateHandle = OnCanvasCreatedHandles.Find(World))
+				{
+					Subsystem->OnGeometryMaskCanvasCreated().Remove(*DelegateHandle);
+				}
+
+				if (FDelegateHandle* DelegateHandle = OnCanvasDestroyedHandles.Find(World))
+				{
+					Subsystem->OnGeometryMaskCanvasDestroyed().Remove(*DelegateHandle);
+				}
 			}
 		}
 	}
@@ -81,22 +133,14 @@ void FGMECanvasListViewModel::OnCanvasCreated(const UGeometryMaskCanvas* InGeome
 	OnChanged().Broadcast();
 }
 
-bool FGMECanvasListViewModel::Tick(const float InDeltaSeconds)
+void FGMECanvasListViewModel::OnCanvasDestroyed(const FGeometryMaskCanvasId& InGeometryMaskCanvasId)
 {
-	if (const UGeometryMaskSubsystem* Subsystem = GEngine->GetEngineSubsystem<UGeometryMaskSubsystem>())
+	CanvasItems.RemoveAll([&InGeometryMaskCanvasId](const TSharedPtr<FGMECanvasItemViewModel>& InItem)
 	{
-		const TArray<FName> CanvasNames = Subsystem->GetCanvasNames();
-		if (LastCanvasNames != CanvasNames)
-		{
-			// Canvas Names changed
-			LastCanvasNames = CanvasNames;
-			RefreshCanvases();
-			OnChanged().Broadcast();
-		}
-	}
-
-	// Always loop
-	return true;
+		return InItem->GetCanvasId() == InGeometryMaskCanvasId;
+	});
+	
+	OnChanged().Broadcast();
 }
 
 bool FGMECanvasListViewModel::GetChildren(TArray<TSharedPtr<IGMETreeNodeViewModel>>& OutChildren)
diff --git a/Engine/Plugins/Experimental/GeometryMask/Source/GeometryMaskEditor/Private/ViewModels/GMECanvasListViewModel.h b/Engine/Plugins/Experimental/GeometryMask/Source/GeometryMaskEditor/Private/ViewModels/GMECanvasListViewModel.h
--- a/Engine/Plugins/Experimental/GeometryMask/Source/GeometryMaskEditor/Private/ViewModels/GMECanvasListViewModel.h
+++ b/Engine/Plugins/Experimental/GeometryMask/Source/GeometryMaskEditor/Private/ViewModels/GMECanvasListViewModel.h
@@ -11,50 +11,42 @@ class FGMECanvasItemViewModel;
 
 class FGMECanvasListViewModel
 	: public TSharedFromThis<FGMECanvasListViewModel>
+	, public FGMEListViewModelBase
 	, public FEditorUndoClient
 	, public IGMETreeNodeViewModel
 {
+	using Super = FGMEListViewModelBase;
+
+protected:
+	// Private token only allows members or friends to call MakeShared
+	struct FPrivateToken { explicit FPrivateToken() = default; };
+	
 public:
-	/**  */
 	static TSharedRef<FGMECanvasListViewModel> Create();
+
+	explicit FGMECanvasListViewModel(FPrivateToken)
+		: Super(Super::FPrivateToken{}) { }
 	virtual ~FGMECanvasListViewModel() override;
 
 	// ~Begin IGMETreeNodeViewModel
 	virtual bool GetChildren(TArray<TSharedPtr<IGMETreeNodeViewModel>>& OutChildren) override;
 	// ~End IGMETreeNodeViewModel
 
-public:
-	using FOnChanged = TMulticastDelegate<void()>;;
-
-	/** Something has changed within the ViewModel */
-	FOnChanged& OnChanged() { return OnChangedDelegate; }
-
 private:
-	// Private token only allows members or friends to call MakeShared
-	struct FPrivateToken { explicit FPrivateToken() = default; };
-	
-public:
-	explicit FGMECanvasListViewModel(FPrivateToken);
+	virtual bool RefreshItems() override;
 
-private:
-	void Initialize();
-
-	void RefreshCanvases();
+	virtual void OnPostWorldInit(UWorld* InWorld, const UWorld::InitializationValues InWorldValues) override;
+	virtual void OnPreWorldDestroyed(UWorld* InWorld) override;
 
 	void OnCanvasCreated(const UGeometryMaskCanvas* InGeometryMaskCanvas);
-
-	bool Tick(const float InDeltaSeconds);
+	void OnCanvasDestroyed(const FGeometryMaskCanvasId& InGeometryMaskCanvasId);
 
 private:
-	/** In seconds. */
-	static constexpr float UpdateCheckInterval = 1.0f;
-	FTSTicker::FDelegateHandle UpdateCheckHandle;
-	
-	FOnChanged OnChangedDelegate;
-	FDelegateHandle OnCanvasCreatedHandle;
+	TMap<TObjectKey<UWorld>, FDelegateHandle> OnCanvasCreatedHandles;
+	TMap<TObjectKey<UWorld>, FDelegateHandle> OnCanvasDestroyedHandles;
 
 	/** Cached canvas names for comparison/refresh. */
-	TArray<FName> LastCanvasNames;
+	TMap<TObjectKey<UWorld>, TArray<FName>> LastCanvasNames;
 
 	/** Canvas ViewModels */
 	TArray<TSharedPtr<FGMECanvasItemViewModel>> CanvasItems;
diff --git a/Engine/Plugins/Experimental/GeometryMask/Source/GeometryMaskEditor/Private/ViewModels/GMEResourceItemViewModel.cpp b/Engine/Plugins/Experimental/GeometryMask/Source/GeometryMaskEditor/Private/ViewModels/GMEResourceItemViewModel.cpp
--- a/Engine/Plugins/Experimental/GeometryMask/Source/GeometryMaskEditor/Private/ViewModels/GMEResourceItemViewModel.cpp
+++ b/Engine/Plugins/Experimental/GeometryMask/Source/GeometryMaskEditor/Private/ViewModels/GMEResourceItemViewModel.cpp
@@ -13,6 +13,19 @@ TSharedRef<FGMEResourceItemViewModel> FGMEResourceItemViewModel::Create(
 	return ViewModel;
 }
 
+FGMEResourceItemViewModel::FGMEResourceItemViewModel(
+	FPrivateToken,
+	const TWeakObjectPtr<const UGeometryMaskCanvasResource>& InResource)
+	: ResourceWeak(InResource)
+{
+	if (const UGeometryMaskCanvasResource* Resource = InResource.Get())
+	{
+		UniqueId = Resource->GetUniqueID();
+		ResourceTextureWeak = const_cast<UGeometryMaskCanvasResource*>(Resource)->GetRenderTargetTexture();
+		UpdateInfoText();
+	}
+}
+
 float FGMEResourceItemViewModel::GetMemoryUsage() const
 {
 	if (const UCanvasRenderTarget2D* Texture = GetResourceTexture())
@@ -33,18 +46,36 @@ FIntPoint FGMEResourceItemViewModel::GetDimensions() const
 	return {0, 0};
 }
 
-FGMEResourceItemViewModel::FGMEResourceItemViewModel(
-	FPrivateToken,
-	const TWeakObjectPtr<const UGeometryMaskCanvasResource>& InResource)
+void FGMEResourceItemViewModel::UpdateInfoText()
 {
-	if (const UGeometryMaskCanvasResource* Resource = InResource.Get())
+	if (const UGeometryMaskCanvasResource* Resource = ResourceWeak.Get())
 	{
-		UniqueId = Resource->GetUniqueID();
-		ResourceTextureWeak = const_cast<UGeometryMaskCanvasResource*>(Resource)->GetRenderTargetTexture();
+		FString UsedChannelsList = FString::JoinBy(Resource->GetDependentCanvasIds(), TEXT("\n"), [](const FGeometryMaskCanvasId& InCanvasId)
+		{
+			return TEXT("\t") + InCanvasId.ToString();
+		});
+		
+		InfoText = FText::FromString(
+    		FString::Printf(TEXT("%-18s: %s\n%-18s: %s\n%-18s: %u of 3\n%s"),
+    			TEXT("Name"), *Resource->GetName(),
+    			TEXT("Size"), *GetDimensions().ToString(),
+    			TEXT("Num. Used Channels"), Resource->GetNumChannelsUsed(),
+    			*UsedChannelsList));
 	}
 }
 
+bool FGMEResourceItemViewModel::Tick(const float InDeltaSeconds)
+{
+	UpdateInfoText();
+	return true;
+}
+
 bool FGMEResourceItemViewModel::GetChildren(TArray<TSharedPtr<IGMETreeNodeViewModel>>& OutChildren)
 {
 	return false;
 }
+
+const FText& FGMEResourceItemViewModel::GetResourceInfo() const
+{
+	return InfoText;
+}
diff --git a/Engine/Plugins/Experimental/GeometryMask/Source/GeometryMaskEditor/Private/ViewModels/GMEResourceItemViewModel.h b/Engine/Plugins/Experimental/GeometryMask/Source/GeometryMaskEditor/Private/ViewModels/GMEResourceItemViewModel.h
--- a/Engine/Plugins/Experimental/GeometryMask/Source/GeometryMaskEditor/Private/ViewModels/GMEResourceItemViewModel.h
+++ b/Engine/Plugins/Experimental/GeometryMask/Source/GeometryMaskEditor/Private/ViewModels/GMEResourceItemViewModel.h
@@ -16,32 +16,38 @@ class UTexture;
 
 class FGMEResourceItemViewModel
 	: public TSharedFromThis<FGMEResourceItemViewModel>
+	, public FGMETickableViewModelBase
 	, public FEditorUndoClient
 	, public IGMETreeNodeViewModel
 {
+protected:
+	// Private token only allows members or friends to call MakeShared
+	struct FPrivateToken { explicit FPrivateToken() = default; };
+	
 public:
-	/**  */
 	static TSharedRef<FGMEResourceItemViewModel> Create(const TWeakObjectPtr<const UGeometryMaskCanvasResource>& InResource);
+
+	FGMEResourceItemViewModel(FPrivateToken, const TWeakObjectPtr<const UGeometryMaskCanvasResource>& InResource);
 	virtual ~FGMEResourceItemViewModel() override = default;
+
+	virtual bool Tick(const float InDeltaSeconds) override;
 	
 	// ~Begin IGMETreeNodeViewModel
 	virtual bool GetChildren(TArray<TSharedPtr<IGMETreeNodeViewModel>>& OutChildren) override;
 	// ~End IGMETreeNodeViewModel
 
 	uint32 GetId() const { return UniqueId; }
-
+	const FText& GetResourceInfo() const;
 	const UCanvasRenderTarget2D* GetResourceTexture() const { return ResourceTextureWeak.Get(); }
 	float GetMemoryUsage() const;
 	FIntPoint GetDimensions() const;
 
 private:
-	// Private token only allows members or friends to call MakeShared
-	struct FPrivateToken { explicit FPrivateToken() = default; };
-
-public:
-	FGMEResourceItemViewModel(FPrivateToken, const TWeakObjectPtr<const UGeometryMaskCanvasResource>& InResource);
+	void UpdateInfoText();
 
 private:
 	uint32 UniqueId;
+	FText InfoText;
+	TWeakObjectPtr<const UGeometryMaskCanvasResource> ResourceWeak;
 	TWeakObjectPtr<UCanvasRenderTarget2D> ResourceTextureWeak;
 };
diff --git a/Engine/Plugins/Experimental/GeometryMask/Source/GeometryMaskEditor/Private/ViewModels/GMEResourceListViewModel.cpp b/Engine/Plugins/Experimental/GeometryMask/Source/GeometryMaskEditor/Private/ViewModels/GMEResourceListViewModel.cpp
--- a/Engine/Plugins/Experimental/GeometryMask/Source/GeometryMaskEditor/Private/ViewModels/GMEResourceListViewModel.cpp
+++ b/Engine/Plugins/Experimental/GeometryMask/Source/GeometryMaskEditor/Private/ViewModels/GMEResourceListViewModel.cpp
@@ -3,7 +3,9 @@
 #include "GMEResourceListViewModel.h"
 
 #include "Engine/Engine.h"
+#include "GeometryMaskCanvasResource.h"
 #include "GeometryMaskSubsystem.h"
+#include "GeometryMaskWorldSubsystem.h"
 #include "GMEResourceItemViewModel.h"
 
 TSharedRef<FGMEResourceListViewModel> FGMEResourceListViewModel::Create()
@@ -18,7 +20,11 @@ FGMEResourceListViewModel::~FGMEResourceListViewModel()
 {
 	if (UGeometryMaskSubsystem* Subsystem = GEngine->GetEngineSubsystem<UGeometryMaskSubsystem>())
 	{
-		Subsystem->OnGeometryMaskCanvasCreated().Remove(OnResourceCreatedHandle);
+		Subsystem->OnGeometryMaskResourceCreated().Remove(OnResourceCreatedHandle);
+		OnResourceCreatedHandle.Reset();
+		
+		Subsystem->OnGeometryMaskResourceDestroyed().Remove(OnResourceDestroyedHandle);
+		OnResourceDestroyedHandle.Reset();
 	}
 	
 	ResourceItems.Reset();
@@ -36,13 +42,16 @@ void FGMEResourceListViewModel::Initialize()
 			ResourceItems.Add(FGMEResourceItemViewModel::Create(Resource));
 		}
 
-		// And listen for new ones
-		{
-			OnResourceCreatedHandle = Subsystem->OnGeometryMaskResourceCreated().AddRaw(this, &FGMEResourceListViewModel::OnResourceCreated);
-		}
-
-		OnChanged().Broadcast();
+		OnResourceCreatedHandle = Subsystem->OnGeometryMaskResourceCreated().AddRaw(this, &FGMEResourceListViewModel::OnResourceCreated);
+		OnResourceDestroyedHandle = Subsystem->OnGeometryMaskResourceDestroyed().AddRaw(this, &FGMEResourceListViewModel::OnResourceDestroyed);
 	}
+
+	FGMEListViewModelBase::Initialize();
+}
+
+bool FGMEResourceListViewModel::RefreshItems()
+{
+	return true;
 }
 
 void FGMEResourceListViewModel::OnResourceCreated(const UGeometryMaskCanvasResource* InGeometryMaskResource)
@@ -61,6 +70,21 @@ void FGMEResourceListViewModel::OnResourceCreated(const UGeometryMaskCanvasResou
 	OnChanged().Broadcast();
 }
 
+void FGMEResourceListViewModel::OnResourceDestroyed(const UGeometryMaskCanvasResource* InGeometryMaskResource)
+{
+	if (!InGeometryMaskResource)
+	{
+		return;
+	}
+	
+	ResourceItems.RemoveAll([InGeometryMaskResource](const TSharedPtr<FGMEResourceItemViewModel>& InViewModel)
+	{
+		return InViewModel->GetId() == InGeometryMaskResource->GetUniqueID();		
+	});
+
+	OnChanged().Broadcast();
+}
+
 bool FGMEResourceListViewModel::GetChildren(TArray<TSharedPtr<IGMETreeNodeViewModel>>& OutChildren)
 {
 	const TArray<TSharedPtr<FGMEResourceItemViewModel>> Children = ResourceItems;
diff --git a/Engine/Plugins/Experimental/GeometryMask/Source/GeometryMaskEditor/Private/ViewModels/GMEResourceListViewModel.h b/Engine/Plugins/Experimental/GeometryMask/Source/GeometryMaskEditor/Private/ViewModels/GMEResourceListViewModel.h
--- a/Engine/Plugins/Experimental/GeometryMask/Source/GeometryMaskEditor/Private/ViewModels/GMEResourceListViewModel.h
+++ b/Engine/Plugins/Experimental/GeometryMask/Source/GeometryMaskEditor/Private/ViewModels/GMEResourceListViewModel.h
@@ -10,39 +10,36 @@ class FGMEResourceItemViewModel;
 
 class FGMEResourceListViewModel
 	: public TSharedFromThis<FGMEResourceListViewModel>
+	, public FGMEListViewModelBase
 	, public FEditorUndoClient
 	, public IGMETreeNodeViewModel
 {
+	using Super = FGMEListViewModelBase;
+
+protected:
+	// Private token only allows members or friends to call MakeShared
+	struct FPrivateToken { explicit FPrivateToken() = default; };
+	
 public:
-	/**  */
 	static TSharedRef<FGMEResourceListViewModel> Create();
+	
+	explicit FGMEResourceListViewModel(FPrivateToken): Super(Super::FPrivateToken{}) { }
 	virtual ~FGMEResourceListViewModel() override;
 
 	// ~Begin IGMETreeNodeViewModel
 	virtual bool GetChildren(TArray<TSharedPtr<IGMETreeNodeViewModel>>& OutChildren) override;
 	// ~End IGMETreeNodeViewModel
 
-	public:
-	using FOnChanged = TMulticastDelegate<void()>;;
-
-	/** Something has changed within the ViewModel */
-	FOnChanged& OnChanged() { return OnChangedDelegate; }
-
-private:
-	// Private token only allows members or friends to call MakeShared
-	struct FPrivateToken { explicit FPrivateToken() = default; };
-	
-public:
-	explicit FGMEResourceListViewModel(FPrivateToken) { }
-
 protected:
-	void Initialize();
+	virtual void Initialize() override;
+	virtual bool RefreshItems() override;
 
 	void OnResourceCreated(const UGeometryMaskCanvasResource* InGeometryMaskResource);
+	void OnResourceDestroyed(const UGeometryMaskCanvasResource* InGeometryMaskResource);
 
 private:
-	FOnChanged OnChangedDelegate;
 	FDelegateHandle OnResourceCreatedHandle;
+	FDelegateHandle OnResourceDestroyedHandle;
 
 	/** Canvas ViewModels */
 	TArray<TSharedPtr<FGMEResourceItemViewModel>> ResourceItems;
diff --git a/Engine/Plugins/Experimental/GeometryMask/Source/GeometryMaskEditor/Private/ViewModels/GMEViewModelShared.cpp b/Engine/Plugins/Experimental/GeometryMask/Source/GeometryMaskEditor/Private/ViewModels/GMEViewModelShared.cpp
new file mode 100644
--- /dev/null
+++ b/Engine/Plugins/Experimental/GeometryMask/Source/GeometryMaskEditor/Private/ViewModels/GMEViewModelShared.cpp
@@ -0,0 +1,66 @@
+﻿// Copyright Epic Games, Inc. All Rights Reserved.
+
+#include "GMEViewModelShared.h"
+
+#include "Engine/Engine.h"
+
+FGMETickableViewModelBase::FGMETickableViewModelBase()
+{
+	UpdateCheckHandle = FTSTicker::GetCoreTicker().AddTicker(FTickerDelegate::CreateRaw(this, &FGMETickableViewModelBase::Tick), UpdateCheckInterval);
+}
+
+FGMETickableViewModelBase::~FGMETickableViewModelBase()
+{
+	FTSTicker::GetCoreTicker().RemoveTicker(UpdateCheckHandle);
+}
+
+FGMEListViewModelBase::FGMEListViewModelBase(FPrivateToken)
+{
+	OnPostWorldInitDelegateHandle = FWorldDelegates::OnPostWorldInitialization.AddRaw(this, &FGMEListViewModelBase::OnPostWorldInit);
+	OnPreWorldDestroyedDelegateHandle = FWorldDelegates::OnPreWorldFinishDestroy.AddRaw(this, &FGMEListViewModelBase::OnPreWorldDestroyed);
+}
+
+FGMEListViewModelBase::~FGMEListViewModelBase()
+{
+	FWorldDelegates::OnPostWorldInitialization.Remove(OnPostWorldInitDelegateHandle);
+	FWorldDelegates::OnPreWorldFinishDestroy.Remove(OnPreWorldDestroyedDelegateHandle);
+}
+
+void FGMEListViewModelBase::Initialize()
+{
+	for (const FWorldContext& WorldContext : GEngine->GetWorldContexts())
+	{
+		if (UWorld* World = WorldContext.World())
+		{
+			LoadedWorlds.Emplace(World);
+		}
+	}
+	
+	RefreshItems();
+}
+
+bool FGMEListViewModelBase::Tick(const float InDeltaSeconds)
+{
+	if (RefreshItems())
+	{
+		OnChanged().Broadcast();
+	}
+
+	return true;
+}
+
+void FGMEListViewModelBase::OnPostWorldInit(UWorld* InWorld, const UWorld::InitializationValues InWorldValues)
+{
+	if (LoadedWorlds.AddUnique(InWorld) == LoadedWorlds.Num() - 1)
+	{
+		RefreshItems();
+	}
+}
+
+void FGMEListViewModelBase::OnPreWorldDestroyed(UWorld* InWorld)
+{
+	if (LoadedWorlds.Remove(InWorld) > 0)
+	{
+		RefreshItems();
+	}
+}
diff --git a/Engine/Plugins/Experimental/GeometryMask/Source/GeometryMaskEditor/Private/ViewModels/GMEViewModelShared.h b/Engine/Plugins/Experimental/GeometryMask/Source/GeometryMaskEditor/Private/ViewModels/GMEViewModelShared.h
--- a/Engine/Plugins/Experimental/GeometryMask/Source/GeometryMaskEditor/Private/ViewModels/GMEViewModelShared.h
+++ b/Engine/Plugins/Experimental/GeometryMask/Source/GeometryMaskEditor/Private/ViewModels/GMEViewModelShared.h
@@ -3,7 +3,10 @@
 #pragma once
 
 #include "Containers/Array.h"
-#include "Templates/SharedPointerFwd.h"
+#include "Containers/Ticker.h"
+#include "Engine/World.h"
+#include "Templates/SharedPointer.h"
+
 
 /** Common interface for treeview nodes. */
 class IGMETreeNodeViewModel
@@ -15,3 +18,59 @@ public:
 	/** See STreeView::OnGetChildren, return false if no children. */
 	virtual bool GetChildren(TArray<TSharedPtr<IGMETreeNodeViewModel>>& OutChildren) = 0;
 };
+
+class FGMETickableViewModelBase
+	: public TSharedFromThis<FGMETickableViewModelBase>
+{
+public:
+	FGMETickableViewModelBase();
+	virtual ~FGMETickableViewModelBase();
+
+protected:
+	virtual bool Tick(const float InDeltaSeconds) { return true; }
+
+protected:
+	/** In seconds. */
+	static constexpr float UpdateCheckInterval = 0.15f;
+	FTSTicker::FDelegateHandle UpdateCheckHandle;
+};
+
+class FGMEListViewModelBase
+	: public TSharedFromThis<FGMEListViewModelBase>
+	, public FGMETickableViewModelBase
+{
+protected:
+	// Private token only allows members or friends to call MakeShared
+	struct FPrivateToken { explicit FPrivateToken() = default; };
+	
+public:
+	explicit FGMEListViewModelBase(FPrivateToken);
+	virtual ~FGMEListViewModelBase() override;
+
+public:
+	using FOnChanged = TMulticastDelegate<void()>;;
+
+	/** Something has changed within the ViewModel */
+	FOnChanged& OnChanged() { return OnChangedDelegate; }
+
+protected:
+	virtual void Initialize();
+	virtual bool RefreshItems() = 0;
+
+	virtual bool Tick(const float InDeltaSeconds) override;
+
+	virtual void OnPostWorldInit(UWorld* InWorld, const UWorld::InitializationValues InWorldValues);
+	virtual void OnPreWorldDestroyed(UWorld* InWorld);
+
+protected:
+	/** In seconds. */
+	static constexpr float UpdateCheckInterval = 1.0f;
+	FTSTicker::FDelegateHandle UpdateCheckHandle;
+	
+	TArray<TObjectKey<UWorld>> LoadedWorlds;
+
+	FDelegateHandle OnPostWorldInitDelegateHandle;
+	FDelegateHandle OnPreWorldDestroyedDelegateHandle;	
+	
+	FOnChanged OnChangedDelegate;
+};
diff --git a/Engine/Plugins/Experimental/GeometryMask/Source/GeometryMaskEditor/Private/Widgets/SGMECanvasItem.cpp b/Engine/Plugins/Experimental/GeometryMask/Source/GeometryMaskEditor/Private/Widgets/SGMECanvasItem.cpp
--- a/Engine/Plugins/Experimental/GeometryMask/Source/GeometryMaskEditor/Private/Widgets/SGMECanvasItem.cpp
+++ b/Engine/Plugins/Experimental/GeometryMask/Source/GeometryMaskEditor/Private/Widgets/SGMECanvasItem.cpp
@@ -21,16 +21,26 @@ void SGMECanvasItem::Construct(
 	ImageWidget = SNew(SGeometryMaskCanvasPreview)
 					.Invert(true)
 					.SolidBackground(true)
-					.CanvasName(ViewModel->GetCanvasName())
+					.CanvasId(ViewModel->GetCanvasId())
 					.Channel(ViewModel->GetColorChannel())
 					.Opacity(1.0f);
 
 	SGMEImageItem::Construct(
 		SGMEImageItem::FArguments()
-			.Label(FText::FromName(ViewModel->GetCanvasName()))
+		.Label(this, &SGMECanvasItem::GetLabel)
 		, InOwnerTableView);
 }
 
+FText SGMECanvasItem::GetLabel() const
+{
+	if (ViewModel.IsValid())
+	{
+		return ViewModel->GetCanvasInfo();
+	}
+
+	return FText::GetEmpty();
+}
+
 FOptionalSize SGMECanvasItem::GetAspectRatio()
 {
 	float AspectRatio = 16.0f / 9.0f;
diff --git a/Engine/Plugins/Experimental/GeometryMask/Source/GeometryMaskEditor/Private/Widgets/SGMECanvasItem.h b/Engine/Plugins/Experimental/GeometryMask/Source/GeometryMaskEditor/Private/Widgets/SGMECanvasItem.h
--- a/Engine/Plugins/Experimental/GeometryMask/Source/GeometryMaskEditor/Private/Widgets/SGMECanvasItem.h
+++ b/Engine/Plugins/Experimental/GeometryMask/Source/GeometryMaskEditor/Private/Widgets/SGMECanvasItem.h
@@ -4,8 +4,6 @@
 
 #include "SGMEImageItem.h"
 #include "Templates/SharedPointer.h"
-#include "Widgets/Views/SListView.h"
-#include "Widgets/Views/STableRow.h"
 
 class FGMECanvasItemViewModel;
 class IGMETreeNodeViewModel;
@@ -23,6 +21,7 @@ public:
 	void Construct(const FArguments& InArgs, const TSharedRef<STableViewBase>& InOwnerTableView, const TSharedRef<FGMECanvasItemViewModel>& InViewModel);
 
 protected:
+	FText GetLabel() const;
 	virtual FOptionalSize GetAspectRatio() override;
 
 private:
diff --git a/Engine/Plugins/Experimental/GeometryMask/Source/GeometryMaskEditor/Private/Widgets/SGMEImageItem.h b/Engine/Plugins/Experimental/GeometryMask/Source/GeometryMaskEditor/Private/Widgets/SGMEImageItem.h
--- a/Engine/Plugins/Experimental/GeometryMask/Source/GeometryMaskEditor/Private/Widgets/SGMEImageItem.h
+++ b/Engine/Plugins/Experimental/GeometryMask/Source/GeometryMaskEditor/Private/Widgets/SGMEImageItem.h
@@ -45,6 +45,7 @@ public:
 			[
 				SNew(SVerticalBox)
 				+ SVerticalBox::Slot()
+				.AutoHeight()
 				.VAlign(VAlign_Top)
 				[
 					SAssignNew(ImageContainerWidget, SBox)
@@ -57,12 +58,14 @@ public:
 				]
 
 				+ SVerticalBox::Slot()
+				.AutoHeight()
 				.Padding(1, 1, 1, Padding)
 				[
 					SNew(SHorizontalBox)
 					+ SHorizontalBox::Slot()
 					[
 						SNew(STextBlock)
+						.Font(FCoreStyle::GetDefaultFontStyle("Mono", 9))
 						.Text(InArgs._Label)
 					]
 				]
@@ -82,7 +85,7 @@ protected:
 			const FGeometry& WidgetGeometry = ImageContainerWidget->GetCachedGeometry();
 			const float AspectRatio = 1.0f / GetAspectRatio().Get();
 			const float Height = static_cast<int32>(WidgetGeometry.GetLocalSize().X) * AspectRatio;
-			return Height * 2; // Don't know why it's x2, it just is.
+			return Height;
 		}
 
 		return {};
diff --git a/Engine/Plugins/Experimental/GeometryMask/Source/GeometryMaskEditor/Private/Widgets/SGMEResourceItem.cpp b/Engine/Plugins/Experimental/GeometryMask/Source/GeometryMaskEditor/Private/Widgets/SGMEResourceItem.cpp
--- a/Engine/Plugins/Experimental/GeometryMask/Source/GeometryMaskEditor/Private/Widgets/SGMEResourceItem.cpp
+++ b/Engine/Plugins/Experimental/GeometryMask/Source/GeometryMaskEditor/Private/Widgets/SGMEResourceItem.cpp
@@ -26,13 +26,20 @@ void SGMEResourceItem::Construct(
 	
 	SGMEImageItem::Construct(
 		SGMEImageItem::FArguments()
-			.Label_Lambda([InViewModel]()
-			{
-				return FText::FromString(InViewModel->GetDimensions().ToString());
-			})
+			.Label(this, &SGMEResourceItem::GetLabel)
 		, InOwnerTableView);
 }
 
+FText SGMEResourceItem::GetLabel() const
+{
+	if (ViewModel.IsValid())
+	{
+		return ViewModel->GetResourceInfo();
+	}
+
+	return FText::GetEmpty();
+}
+
 FOptionalSize SGMEResourceItem::GetAspectRatio()
 {
 	float AspectRatio = 16.0f / 9.0f;
diff --git a/Engine/Plugins/Experimental/GeometryMask/Source/GeometryMaskEditor/Private/Widgets/SGMEResourceItem.h b/Engine/Plugins/Experimental/GeometryMask/Source/GeometryMaskEditor/Private/Widgets/SGMEResourceItem.h
--- a/Engine/Plugins/Experimental/GeometryMask/Source/GeometryMaskEditor/Private/Widgets/SGMEResourceItem.h
+++ b/Engine/Plugins/Experimental/GeometryMask/Source/GeometryMaskEditor/Private/Widgets/SGMEResourceItem.h
@@ -23,6 +23,7 @@ public:
 	void Construct(const FArguments& InArgs, const TSharedRef<STableViewBase>& InOwnerTableView, const TSharedRef<FGMEResourceItemViewModel>& InViewModel);
 
 protected:
+	FText GetLabel() const;
 	virtual FOptionalSize GetAspectRatio() override;
 	
 private:
diff --git a/Engine/Plugins/Experimental/GeometryMask/Source/GeometryMaskEditor/Private/Widgets/SGeometryMaskCanvasPreview.cpp b/Engine/Plugins/Experimental/GeometryMask/Source/GeometryMaskEditor/Private/Widgets/SGeometryMaskCanvasPreview.cpp
--- a/Engine/Plugins/Experimental/GeometryMask/Source/GeometryMaskEditor/Private/Widgets/SGeometryMaskCanvasPreview.cpp
+++ b/Engine/Plugins/Experimental/GeometryMask/Source/GeometryMaskEditor/Private/Widgets/SGeometryMaskCanvasPreview.cpp
@@ -2,12 +2,15 @@
 
 #include "Widgets/SGeometryMaskCanvasPreview.h"
 
+#include "EditorSupportDelegates.h"
 #include "Engine/CanvasRenderTarget2D.h"
 #include "Engine/Engine.h"
 #include "Engine/Texture.h"
+#include "GeometryMaskCanvasResource.h"
 #include "GeometryMaskEditorLog.h"
 #include "GeometryMaskSubsystem.h"
 #include "GeometryMaskTypes.h"
+#include "GeometryMaskWorldSubsystem.h"
 #include "Materials/MaterialInstanceDynamic.h"
 #include "Materials/MaterialInterface.h"
 #include "SlateMaterialBrush.h"
@@ -38,12 +41,18 @@ SGeometryMaskCanvasPreview::SGeometryMaskCanvasPreview()
 {
 }
 
+SGeometryMaskCanvasPreview::~SGeometryMaskCanvasPreview()
+{
+	FEditorSupportDelegates::PrepareToCleanseEditorObject.RemoveAll(this);
+}
+
 void SGeometryMaskCanvasPreview::Construct(const FArguments& InArgs)
 {
-	CanvasName = InArgs._CanvasName;
+	CanvasId = InArgs._CanvasId;
 	ColorChannel = InArgs._Channel;
-	bInvert = InArgs._Invert;
-	bHasSolidBackground = InArgs._SolidBackground;
+	bShowPaddingFrame = InArgs._PaddingFrameVisibility;
+	Invert = InArgs._Invert;
+	HasSolidBackground = InArgs._SolidBackground;
 	OpacityMultiplier = InArgs._Opacity;
 	AspectRatio = 1920.0f / 1080.0f;
 
@@ -84,7 +93,7 @@ void SGeometryMaskCanvasPreview::Construct(const FArguments& InArgs)
 				SNew(SImage)
 				.Visibility_Lambda([this]()
 				{
-					return HasSolidBackground()
+					return GetSolidBackground()
 						? EVisibility::Visible
 						: EVisibility::Hidden;
 				})
@@ -107,20 +116,25 @@ void SGeometryMaskCanvasPreview::Tick(const FGeometry& InAllottedGeometry, const
 	UpdateBrush(CanvasWeak.Get(), nullptr);
 }
 
-const FName SGeometryMaskCanvasPreview::GetCanvasName() const
+const FGeometryMaskCanvasId& SGeometryMaskCanvasPreview::GetCanvasId() const
 {
-	return CanvasName.Get(NAME_None);
+	return CanvasId.Get(FGeometryMaskCanvasId::None);
 }
 
-void SGeometryMaskCanvasPreview::SetCanvasName(const FName InCanvasName)
+void SGeometryMaskCanvasPreview::SetCanvasId(const FGeometryMaskCanvasId& InCanvasId)
 {
-	if (GetCanvasName() != InCanvasName)
+	if (GetCanvasId() != InCanvasId)
 	{
-		CanvasName.Set(InCanvasName);
+		CanvasId.Set(InCanvasId);
 		TryResolveCanvas();
 	}
 }
 
+const FName SGeometryMaskCanvasPreview::GetCanvasName() const
+{
+	return GetCanvasId().Name;
+}
+
 const EGeometryMaskColorChannel SGeometryMaskCanvasPreview::GetColorChannel(const bool bOnlyValid) const
 {
 	EGeometryMaskColorChannel Result = ColorChannel.Get(EGeometryMaskColorChannel::Red);
@@ -147,24 +161,34 @@ void SGeometryMaskCanvasPreview::SetColorChannel(const EGeometryMaskColorChannel
 	}
 }
 
+const bool SGeometryMaskCanvasPreview::IsPaddingFrameVisible() const
+{
+	return bShowPaddingFrame.Get(true);
+}
+
+void SGeometryMaskCanvasPreview::SetPaddingFrameVisiblity(const bool bInInverted)
+{
+	bShowPaddingFrame.Set(bInInverted);
+}
+
 const bool SGeometryMaskCanvasPreview::IsInverted() const
 {
-	return bInvert.Get(false);
+	return Invert.Get(false);
 }
 
 void SGeometryMaskCanvasPreview::SetInvert(const bool bInInverted)
 {
-	bInvert.Set(bInInverted);
+	Invert.Set(bInInverted);
 }
 
-const bool SGeometryMaskCanvasPreview::HasSolidBackground() const
+const bool SGeometryMaskCanvasPreview::GetSolidBackground() const
 {
-	return bHasSolidBackground.Get();
+	return HasSolidBackground.Get();
 }
 
 void SGeometryMaskCanvasPreview::SetSolidBackground(const bool bInHasSolidBackground)
 {
-	bHasSolidBackground.Set(bInHasSolidBackground);
+	HasSolidBackground.Set(bInHasSolidBackground);
 }
 
 const float SGeometryMaskCanvasPreview::GetOpacity() const
@@ -220,10 +244,18 @@ FString SGeometryMaskCanvasPreview::GetReferencerName() const
 
 bool SGeometryMaskCanvasPreview::TryResolveCanvas()
 {
-	if (UGeometryMaskSubsystem* Subsystem = GEngine->GetEngineSubsystem<UGeometryMaskSubsystem>())
+	if (!CanvasId.IsSet())
+	{
+		return false;
+	}
+
+	if (UWorld* CanvasWorld = CanvasId.Get().World.ResolveObjectPtr())
 	{
-		CanvasWeak = Subsystem->GetNamedCanvas(GetCanvasName());
-		UpdateBrush(CanvasWeak.Get(), nullptr);
+		if (UGeometryMaskWorldSubsystem* Subsystem = CanvasWorld->GetSubsystem<UGeometryMaskWorldSubsystem>())
+		{
+			CanvasWeak = Subsystem->GetNamedCanvas(GetCanvasName());
+			UpdateBrush(CanvasWeak.Get(), nullptr);
+		}
 	}
 
 	return CanvasWeak != nullptr;
@@ -247,11 +279,13 @@ void SGeometryMaskCanvasPreview::UpdateBrush(const UGeometryMaskCanvas* InCanvas
 			InTexture = CanvasTexture;
 		}
 
+		FGeometryMaskDrawingContext DrawingContext(InCanvas->GetCanvasId().World);
+
 		FVector4f Padding(ForceInitToZero);
 		if (const UGeometryMaskCanvasResource* CanvasResource = InCanvas->GetResource())
 		{
-			FIntPoint ViewportPadding(CanvasResource->GetViewportPadding());
-			FVector2f ViewportPaddingF = FVector2f(ViewportPadding) / CanvasResource->GetViewportSize();
+			FIntPoint ViewportPadding(CanvasResource->GetViewportPadding(DrawingContext));
+			FVector2f ViewportPaddingF = FVector2f(ViewportPadding) / CanvasResource->GetMaxViewportSize();
 			Padding.X = ViewportPaddingF.X;
 			Padding.Y = ViewportPaddingF.Y;
 		}
diff --git a/Engine/Plugins/Experimental/GeometryMask/Source/GeometryMaskEditor/Public/Widgets/SGeometryMaskCanvasPreview.h b/Engine/Plugins/Experimental/GeometryMask/Source/GeometryMaskEditor/Public/Widgets/SGeometryMaskCanvasPreview.h
--- a/Engine/Plugins/Experimental/GeometryMask/Source/GeometryMaskEditor/Public/Widgets/SGeometryMaskCanvasPreview.h
+++ b/Engine/Plugins/Experimental/GeometryMask/Source/GeometryMaskEditor/Public/Widgets/SGeometryMaskCanvasPreview.h
@@ -18,32 +18,41 @@ class GEOMETRYMASKEDITOR_API SGeometryMaskCanvasPreview
 {
 public:
 	SLATE_BEGIN_ARGS(SGeometryMaskCanvasPreview) {}
-		SLATE_ATTRIBUTE(FName, CanvasName)
+		SLATE_ATTRIBUTE(FGeometryMaskCanvasId, CanvasId)
 		SLATE_ATTRIBUTE(EGeometryMaskColorChannel, Channel)
+		SLATE_ATTRIBUTE(bool, PaddingFrameVisibility)
 		SLATE_ATTRIBUTE(bool, Invert)
 		SLATE_ATTRIBUTE(bool, SolidBackground)
 		SLATE_ATTRIBUTE(float, Opacity)
 	SLATE_END_ARGS()
 
 	SGeometryMaskCanvasPreview();
+	virtual ~SGeometryMaskCanvasPreview() override;
 
 	/** Constructs this widget with InArgs */
 	void Construct(const FArguments& InArgs);
 
 	virtual void Tick(const FGeometry& InAllottedGeometry, const double InCurrentTime, const float InDeltaTime) override;
 
+	/** Get the id of the currently referenced canvas. */
+	const FGeometryMaskCanvasId& GetCanvasId() const;
+
+	/** Sets the id of the currently referenced canvas, and resolves the canvas itself. */
+	void SetCanvasId(const FGeometryMaskCanvasId& InCanvasId);
+
 	/** Get the name of the currently referenced canvas. */
 	const FName GetCanvasName() const;
 
-	/** Sets the name of the currently referenced canvas, and resolves the canvas itself. */
-	void SetCanvasName(const FName InCanvasName);
-
 	/** Get the color channel to display for the canvas.  */
 	const EGeometryMaskColorChannel GetColorChannel(const bool bOnlyValid = true) const;
 
 	/** Set the color channel to display for the canvas.  */
 	void SetColorChannel(const EGeometryMaskColorChannel InColorChannel);
 
+	const bool IsPaddingFrameVisible() const;
+
+	void SetPaddingFrameVisiblity(const bool bInIsVisible);
+
 	/** Get whether to invert the display of the canvas.  */
 	const bool IsInverted() const;
 
@@ -51,7 +60,7 @@ public:
 	void SetInvert(const bool bInInverted);
 
 	/** Get whether a solid background is used or not (vs. compositing with alpha on widgets below). */
-	const bool HasSolidBackground() const;
+	const bool GetSolidBackground() const;
 
 	/** Set whether a solid background is used or not (vs. compositing with alpha on widgets below).  */
 	void SetSolidBackground(const bool bInHasSolidBackground);
@@ -88,17 +97,19 @@ private:
 
 	TWeakObjectPtr<UGeometryMaskCanvas> CanvasWeak;
 
-	/** Name of the currently referenced canvas. */
-	TAttribute<FName> CanvasName;
+	/** Id of the currently referenced canvas. */
+	TAttribute<FGeometryMaskCanvasId> CanvasId;
 
 	/** Color channel to display for the canvas.  */
 	TAttribute<EGeometryMaskColorChannel> ColorChannel;
 
+	TAttribute<bool> bShowPaddingFrame;
+
 	/** Whether to invert the display of the canvas.  */
-	TAttribute<bool> bInvert;
+	TAttribute<bool> Invert;
 
 	/*** Whether a solid background is used or not (vs. compositing with alpha on widgets below). */
-	TAttribute<bool> bHasSolidBackground;
+	TAttribute<bool> HasSolidBackground;
 
 	/** Multiplies the overall opacity. */
 	TAttribute<float> OpacityMultiplier;
