GeometryMask: Fix logs being flooded by setting the default verbosity of LogGeometryMask to Log

diff --git a/Engine/Plugins/Experimental/GeometryMask/Source/GeometryMask/Private/GeometryMaskModule.h b/Engine/Plugins/Experimental/GeometryMask/Source/GeometryMask/Private/GeometryMaskModule.h
--- a/Engine/Plugins/Experimental/GeometryMask/Source/GeometryMask/Private/GeometryMaskModule.h
+++ b/Engine/Plugins/Experimental/GeometryMask/Source/GeometryMask/Private/GeometryMaskModule.h
@@ -8,7 +8,7 @@
 #include "Modules/ModuleInterface.h"
 #include "Stats/Stats.h"
 
-DECLARE_LOG_CATEGORY_EXTERN(LogGeometryMask, All, All);
+DECLARE_LOG_CATEGORY_EXTERN(LogGeometryMask, Log, All);
 DECLARE_STATS_GROUP(TEXT("GeometryMask"), STATGROUP_GeometryMask, STATCAT_Advanced);
 
 class FGeometryMaskModule
