[GeometryMask] Async access issue

diff --git a/Engine/Plugins/Experimental/GeometryMask/Source/GeometryMask/Private/GeometryMaskWriteComponent.cpp b/Engine/Plugins/Experimental/GeometryMask/Source/GeometryMask/Private/GeometryMaskWriteComponent.cpp
--- a/Engine/Plugins/Experimental/GeometryMask/Source/GeometryMask/Private/GeometryMaskWriteComponent.cpp
+++ b/Engine/Plugins/Experimental/GeometryMask/Source/GeometryMask/Private/GeometryMaskWriteComponent.cpp
@@ -3,7 +3,6 @@
 #include "GeometryMaskWriteComponent.h"
 
 #include "Algo/RemoveIf.h"
-#include "Async/ParallelFor.h"
 #include "BatchedElements.h"
 #include "CanvasTypes.h"
 #include "Components/DynamicMeshComponent.h"
@@ -71,29 +70,33 @@ void UGeometryMaskWriteMeshComponent::DrawToCanvas(FCanvas* InCanvas)
 		TMap<FName, TWeakObjectPtr<USceneComponent>> ComponentsToKeep;
 		ComponentsToKeep.Reserve(KeyNames.Num());
 
+		TMap<FName, FTransform> ComponentTransforms;
+		ComponentTransforms.Reserve(KeyNames.Num());
+
 		TRACE_CPUPROFILER_EVENT_SCOPE(UGeometryMaskWriteMeshComponent::DrawToCanvas::WriteToCanvas);
 
-		ParallelFor(KeyNames.Num(), [&](int32 TaskIdx)
+		// Flag valid components, can't do async due to UObject access
+		for (FName KeyName : KeyNames)
 		{
-			TRACE_CPUPROFILER_EVENT_SCOPE(UGeometryMaskWriteMeshComponent::DrawToCanvas::WriteToCanvas::Task);
-
-			const FName KeyName = KeyNames[TaskIdx];
-			FTransform LocalToWorld;
-
 			if (CachedComponentsWeak.Contains(KeyName))
 			{
 				TWeakObjectPtr<USceneComponent>& CachedComponentWeak = CachedComponentsWeak[KeyName];
-				if (const USceneComponent* Component = CachedComponentWeak.Get())
+				if (const USceneComponent* CachedComponent = CachedComponentWeak.Get())
 				{
 					ComponentsToKeep.Emplace(KeyName, CachedComponentWeak);
-					LocalToWorld = Component->GetComponentToWorld();
+					ComponentTransforms.Emplace(KeyName, CachedComponent->GetComponentToWorld());
 				}
 			}
-			else
-			{
-				return;
-			}
+		}
+
+		// Regenerate keys for valid components only
+		ComponentsToKeep.GenerateKeyArray(KeyNames);
 
+		for (FName KeyName : KeyNames)
+		{
+			TRACE_CPUPROFILER_EVENT_SCOPE(UGeometryMaskWriteMeshComponent::DrawToCanvas::WriteToCanvas::Task);
+
+			const FTransform LocalToWorld = ComponentTransforms[KeyName];
 			const FGeometryMaskBatchElementData& MeshBatchElementData = CachedMeshData[KeyName];
 
 			{
@@ -145,7 +148,7 @@ void UGeometryMaskWriteMeshComponent::DrawToCanvas(FCanvas* InCanvas)
 				FScopeLock CanvasLock(&CanvasCS);
 				InCanvas->PopTransform();
 			}
-		});
+		}
 
 		CachedComponentsWeak = ComponentsToKeep;
 	}
@@ -276,17 +279,19 @@ void UGeometryMaskWriteMeshComponent::UpdateCachedStaticMeshData(TConstArrayView
 		// Convert mesh resources to batch elements
 		{
 			TRACE_CPUPROFILER_EVENT_SCOPE(UGeometryMaskWriteMeshComponent::DrawToCanvas::ConvertMeshResources);
-			ParallelFor(StaticMeshResources.Num(),
-			[this, &StaticMeshResources, &StaticMeshObjectNames](int32 TaskIdx)
+			
+			CachedMeshData.Reserve(CachedMeshData.Num() + StaticMeshObjectNames.Num());
+
+			for (int32 MeshIdx = 0; MeshIdx < StaticMeshResources.Num(); ++MeshIdx)
 			{
 				TRACE_CPUPROFILER_EVENT_SCOPE(UGeometryMaskWriteMeshComponent::DrawToCanvas::ConvertMeshResources::Task);
-				
-				FStaticMeshLODResources* MeshResources = StaticMeshResources[TaskIdx];
+
+				FStaticMeshLODResources* MeshResources = StaticMeshResources[MeshIdx];
 				const int32 NumVertices = MeshResources->GetNumVertices();
 				const int32 NumIndices = MeshResources->IndexBuffer.GetNumIndices();
 				const int32 NumTriangles = MeshResources->GetNumTriangles();
 
-				FGeometryMaskBatchElementData& MeshBatchElementData = CachedMeshData.Emplace(StaticMeshObjectNames[TaskIdx]);
+				FGeometryMaskBatchElementData& MeshBatchElementData = CachedMeshData.Emplace(StaticMeshObjectNames[MeshIdx]);
 				MeshBatchElementData.Reserve(NumVertices, NumIndices, NumTriangles);
 				MeshResources->IndexBuffer.GetCopy(MeshBatchElementData.Indices);
 
@@ -294,7 +299,7 @@ void UGeometryMaskWriteMeshComponent::UpdateCachedStaticMeshData(TConstArrayView
 				{
 					MeshBatchElementData.Vertices.Add(FVector4f(MeshResources->VertexBuffers.PositionVertexBuffer.VertexPosition(VertexIdx)));
 				}
-			});
+			}
 		}
 	}
 }
@@ -381,20 +386,22 @@ void UGeometryMaskWriteMeshComponent::UpdateCachedDynamicMeshData(TConstArrayVie
 	// Convert meshes to batch elements
 	{
 		TRACE_CPUPROFILER_EVENT_SCOPE(UGeometryMaskWriteMeshComponent::DrawToCanvas::ConvertMeshResources);
-		ParallelFor(DynamicMeshes.Num(),
-		[this, &DynamicMeshes, &DynamicMeshObjectNames](int32 TaskIdx)
+		
+		CachedMeshData.Reserve(CachedMeshData.Num() + DynamicMeshObjectNames.Num());
+		
+		for (int32 MeshIdx = 0; MeshIdx < DynamicMeshes.Num(); ++MeshIdx)
 		{
 			TRACE_CPUPROFILER_EVENT_SCOPE(UGeometryMaskWriteMeshComponent::DrawToCanvas::ConvertMeshResources::Task);
 
 			FDynamicMesh3 CompactMesh;
-			const FDynamicMesh3* DynamicMesh = DynamicMeshes[TaskIdx];
+			const FDynamicMesh3* DynamicMesh = DynamicMeshes[MeshIdx];
 			CompactMesh.CompactCopy(*DynamicMesh);
 
 			const int32 NumVertices = CompactMesh.VertexCount();
 			const int32 NumIndices = CompactMesh.TriangleCount() * 3;
 			const int32 NumTriangles = CompactMesh.TriangleCount();
 
-			FGeometryMaskBatchElementData& MeshBatchElementData = CachedMeshData.Emplace(DynamicMeshObjectNames[TaskIdx]);
+			FGeometryMaskBatchElementData& MeshBatchElementData = CachedMeshData.Emplace(DynamicMeshObjectNames[MeshIdx]);
 			MeshBatchElementData.ChangeStamp = DynamicMesh->GetChangeStamp();
 			MeshBatchElementData.Reserve(NumVertices, NumIndices, NumTriangles);
 
@@ -410,7 +417,7 @@ void UGeometryMaskWriteMeshComponent::UpdateCachedDynamicMeshData(TConstArrayVie
 					static_cast<uint32>(Triangle.B),
 					static_cast<uint32>(Triangle.C)});
 			}
-		});
+		}
 	}
 }
 
@@ -418,7 +425,7 @@ void UGeometryMaskWriteMeshComponent::UpdateCachedDynamicMeshData(TConstArrayVie
 void UGeometryMaskWriteMeshComponent::OnStaticMeshChanged(UStaticMeshComponent* InStaticMeshComponent)
 {
 	// Triggers a cache refresh
-	LastPrimitiveComponentCount = -1;
+	ResetCachedData();
 }
 #endif
 
@@ -431,7 +438,7 @@ void UGeometryMaskWriteMeshComponent::OnDynamicMeshChanged(UDynamicMeshComponent
 	}
 	
 	// Triggers a general cache refresh
-	LastPrimitiveComponentCount = -1;
+	ResetCachedData();
 }
 
 bool UGeometryMaskWriteMeshComponent::TryResolveCanvas()
@@ -466,3 +473,10 @@ bool UGeometryMaskWriteMeshComponent::Cleanup()
 
 	return false;
 }
+
+void UGeometryMaskWriteMeshComponent::ResetCachedData()
+{
+	LastPrimitiveComponentCount = -1;
+	CachedComponentsWeak.Reset();
+	CachedMeshData.Reset();
+}
diff --git a/Engine/Plugins/Experimental/GeometryMask/Source/GeometryMask/Public/GeometryMaskWriteComponent.h b/Engine/Plugins/Experimental/GeometryMask/Source/GeometryMask/Public/GeometryMaskWriteComponent.h
--- a/Engine/Plugins/Experimental/GeometryMask/Source/GeometryMask/Public/GeometryMaskWriteComponent.h
+++ b/Engine/Plugins/Experimental/GeometryMask/Source/GeometryMask/Public/GeometryMaskWriteComponent.h
@@ -72,7 +72,10 @@ protected:
     virtual bool TryResolveCanvas() override;
 	virtual bool Cleanup() override;
     //~ End UGeometryMaskCanvasReferenceComponentBase
-    	
+
+	// Resets cached data, triggers rebuild
+	void ResetCachedData();
+	
 	void UpdateCachedData();
 	void UpdateCachedStaticMeshData(TConstArrayView<UPrimitiveComponent*> InPrimitiveComponents);
 	void UpdateCachedDynamicMeshData(TConstArrayView<UPrimitiveComponent*> InPrimitiveComponents);
