[MotionDesign] Missing initializers

diff --git a/Engine/Plugins/Experimental/GeometryMask/Source/GeometryMask/Public/GeometryMaskTypes.h b/Engine/Plugins/Experimental/GeometryMask/Source/GeometryMask/Public/GeometryMaskTypes.h
--- a/Engine/Plugins/Experimental/GeometryMask/Source/GeometryMask/Public/GeometryMaskTypes.h
+++ b/Engine/Plugins/Experimental/GeometryMask/Source/GeometryMask/Public/GeometryMaskTypes.h
@@ -128,11 +128,11 @@ public:
 	uint8 SceneViewIndex = 0;
 
 	UPROPERTY()
-	FIntPoint ViewportSize;
+	FIntPoint ViewportSize = { 1920, 1080 };
 
 	/** The last resolved ViewProjectionMatrix. */
 	UPROPERTY()
-	FMatrix ViewProjectionMatrix;
+	FMatrix ViewProjectionMatrix = FMatrix::Identity;
 
 public:
 	FGeometryMaskDrawingContext() = default;
