Geometry Mask: fix error X3500: array reference cannot be used as an l-value; not natively addressable

diff --git a/Engine/Plugins/Experimental/GeometryMask/Shaders/Private/GeometryMaskJFtoDFCS.usf b/Engine/Plugins/Experimental/GeometryMask/Shaders/Private/GeometryMaskJFtoDFCS.usf
--- a/Engine/Plugins/Experimental/GeometryMask/Shaders/Private/GeometryMaskJFtoDFCS.usf
+++ b/Engine/Plugins/Experimental/GeometryMask/Shaders/Private/GeometryMaskJFtoDFCS.usf
@@ -38,7 +38,7 @@ void MainCS(
 
 #if DEBUG_SOBEL
 	float4 Sample = InputBuffer[PixelCoord1D];
-	Outline[0] = saturate(distance(UV, Sample.xy));
+	Outline.x = saturate(distance(UV, Sample.xy));
 #else
 	// Perform on each of R,G,B,A
 	[unroll]
@@ -48,13 +48,21 @@ void MainCS(
 		
 		float4 Sample = InputBuffer[PixelCoord1D + ChannelIdx];
 
-		Mask[OutputChannelIdx] = (Sample.x + Sample.y + Sample.a > 0.001);
+		float MaskValue = (Sample.x + Sample.y + Sample.a > 0.001);
 
-		Outline[OutputChannelIdx] = saturate(distance(UV, Sample.xy) * StepDistanceMultipliers[OutputChannelIdx]);
+		float OutlineValue = saturate(distance(UV, Sample.xy) * StepDistanceMultipliers[OutputChannelIdx]);
 
 		// Set range such that inner goes from 0.0-0.5, outer goes from 0.5-1.0
-		Outline[OutputChannelIdx] = mad(Outline[OutputChannelIdx], 0.5, 0.5);
-		Outline[OutputChannelIdx] = 1.0 - abs(Outline[OutputChannelIdx] - saturate(Sample.a));
+		OutlineValue = mad(OutlineValue, 0.5, 0.5);
+		OutlineValue = 1.0 - abs(OutlineValue - saturate(Sample.a));
+
+		switch (OutputChannelIdx)
+		{
+		case 0: Mask.x = MaskValue; Outline.x = OutlineValue; break;
+		case 1: Mask.y = MaskValue; Outline.y = OutlineValue; break;
+		case 2: Mask.z = MaskValue; Outline.z = OutlineValue; break;
+		case 3: Mask.w = MaskValue; Outline.w = OutlineValue; break;
+		}
 	}
 #endif
 
