Fixed multiply-inherited TSharedFromThis.

diff --git a/Engine/Plugins/Experimental/Avalanche/Source/AvalancheMediaEditor/Private/Rundown/Pages/Slate/SAvaRundownPageViewRow.h b/Engine/Plugins/Experimental/Avalanche/Source/AvalancheMediaEditor/Private/Rundown/Pages/Slate/SAvaRundownPageViewRow.h
--- a/Engine/Plugins/Experimental/Avalanche/Source/AvalancheMediaEditor/Private/Rundown/Pages/Slate/SAvaRundownPageViewRow.h
+++ b/Engine/Plugins/Experimental/Avalanche/Source/AvalancheMediaEditor/Private/Rundown/Pages/Slate/SAvaRundownPageViewRow.h
@@ -26,7 +26,7 @@ enum class EAvaRundownRowState
 	Disabled
 };
 
-class FAvaRundownPageViewRowDragDropOp : public FDragDropOperation, public TSharedFromThis<FAvaRundownPageViewRowDragDropOp>
+class FAvaRundownPageViewRowDragDropOp : public FDragDropOperation
 {
 public:
 	DRAG_DROP_OPERATOR_TYPE(FAvaRundownPageViewRowDragDropOp, FDragDropOperation)
diff --git a/Engine/Plugins/Experimental/Avalanche/Source/AvalancheOutliner/Private/Columns/Slate/SAvaOutlinerColor.cpp b/Engine/Plugins/Experimental/Avalanche/Source/AvalancheOutliner/Private/Columns/Slate/SAvaOutlinerColor.cpp
--- a/Engine/Plugins/Experimental/Avalanche/Source/AvalancheOutliner/Private/Columns/Slate/SAvaOutlinerColor.cpp
+++ b/Engine/Plugins/Experimental/Avalanche/Source/AvalancheOutliner/Private/Columns/Slate/SAvaOutlinerColor.cpp
@@ -119,7 +119,7 @@ public:
 	}
 };
 
-class FColorDragDropOp : public FDragDropOperation, public TSharedFromThis<FColorDragDropOp>
+class FColorDragDropOp : public FDragDropOperation
 {
 public:
 	DRAG_DROP_OPERATOR_TYPE(FColorDragDropOp, FDragDropOperation)
diff --git a/Engine/Plugins/Experimental/Avalanche/Source/AvalancheOutliner/Private/Columns/Slate/SAvaOutlinerLock.cpp b/Engine/Plugins/Experimental/Avalanche/Source/AvalancheOutliner/Private/Columns/Slate/SAvaOutlinerLock.cpp
--- a/Engine/Plugins/Experimental/Avalanche/Source/AvalancheOutliner/Private/Columns/Slate/SAvaOutlinerLock.cpp
+++ b/Engine/Plugins/Experimental/Avalanche/Source/AvalancheOutliner/Private/Columns/Slate/SAvaOutlinerLock.cpp
@@ -10,7 +10,7 @@
 
 namespace UE::AvaOutliner::Private
 {
-	class FLockDragDropOp : public FDragDropOperation, public TSharedFromThis<FLockDragDropOp>
+	class FLockDragDropOp : public FDragDropOperation
 	{
 	public:
 		DRAG_DROP_OPERATOR_TYPE(FLockDragDropOp, FDragDropOperation)
diff --git a/Engine/Plugins/Experimental/Avalanche/Source/AvalancheOutliner/Private/Columns/Slate/SAvaOutlinerVisibility.cpp b/Engine/Plugins/Experimental/Avalanche/Source/AvalancheOutliner/Private/Columns/Slate/SAvaOutlinerVisibility.cpp
--- a/Engine/Plugins/Experimental/Avalanche/Source/AvalancheOutliner/Private/Columns/Slate/SAvaOutlinerVisibility.cpp
+++ b/Engine/Plugins/Experimental/Avalanche/Source/AvalancheOutliner/Private/Columns/Slate/SAvaOutlinerVisibility.cpp
@@ -12,7 +12,7 @@
 
 namespace UE::AvaOutliner::Private
 {
-	class FVisibilityDragDropOp : public FDragDropOperation, public TSharedFromThis<FVisibilityDragDropOp>
+	class FVisibilityDragDropOp : public FDragDropOperation
 	{
 	public:
 		DRAG_DROP_OPERATOR_TYPE(FVisibilityDragDropOp, FDragDropOperation)
diff --git a/Engine/Plugins/Experimental/GeometryMask/Source/GeometryMaskEditor/Private/ViewModels/GMECanvasItemViewModel.h b/Engine/Plugins/Experimental/GeometryMask/Source/GeometryMaskEditor/Private/ViewModels/GMECanvasItemViewModel.h
--- a/Engine/Plugins/Experimental/GeometryMask/Source/GeometryMaskEditor/Private/ViewModels/GMECanvasItemViewModel.h
+++ b/Engine/Plugins/Experimental/GeometryMask/Source/GeometryMaskEditor/Private/ViewModels/GMECanvasItemViewModel.h
@@ -1,4 +1,4 @@
-﻿// Copyright Epic Games, Inc. All Rights Reserved.
+// Copyright Epic Games, Inc. All Rights Reserved.
 
 #pragma once
 
@@ -15,8 +15,7 @@ class UGeometryMaskCanvas;
 class UTexture;
 
 class FGMECanvasItemViewModel
-	: public TSharedFromThis<FGMECanvasItemViewModel>
-	, public FGMETickableViewModelBase
+	: public FGMETickableViewModelBase
 	, public FEditorUndoClient
 	, public IGMETreeNodeViewModel
 {
diff --git a/Engine/Plugins/Experimental/GeometryMask/Source/GeometryMaskEditor/Private/ViewModels/GMECanvasListViewModel.h b/Engine/Plugins/Experimental/GeometryMask/Source/GeometryMaskEditor/Private/ViewModels/GMECanvasListViewModel.h
--- a/Engine/Plugins/Experimental/GeometryMask/Source/GeometryMaskEditor/Private/ViewModels/GMECanvasListViewModel.h
+++ b/Engine/Plugins/Experimental/GeometryMask/Source/GeometryMaskEditor/Private/ViewModels/GMECanvasListViewModel.h
@@ -1,4 +1,4 @@
-﻿// Copyright Epic Games, Inc. All Rights Reserved.
+// Copyright Epic Games, Inc. All Rights Reserved.
 
 #pragma once
 
@@ -10,8 +10,7 @@
 class FGMECanvasItemViewModel;
 
 class FGMECanvasListViewModel
-	: public TSharedFromThis<FGMECanvasListViewModel>
-	, public FGMEListViewModelBase
+	: public FGMEListViewModelBase
 	, public FEditorUndoClient
 	, public IGMETreeNodeViewModel
 {
diff --git a/Engine/Plugins/Experimental/GeometryMask/Source/GeometryMaskEditor/Private/ViewModels/GMEResourceItemViewModel.h b/Engine/Plugins/Experimental/GeometryMask/Source/GeometryMaskEditor/Private/ViewModels/GMEResourceItemViewModel.h
--- a/Engine/Plugins/Experimental/GeometryMask/Source/GeometryMaskEditor/Private/ViewModels/GMEResourceItemViewModel.h
+++ b/Engine/Plugins/Experimental/GeometryMask/Source/GeometryMaskEditor/Private/ViewModels/GMEResourceItemViewModel.h
@@ -1,4 +1,4 @@
-﻿// Copyright Epic Games, Inc. All Rights Reserved.
+// Copyright Epic Games, Inc. All Rights Reserved.
 
 #pragma once
 
@@ -15,8 +15,7 @@ class UGeometryMaskCanvas;
 class UTexture;
 
 class FGMEResourceItemViewModel
-	: public TSharedFromThis<FGMEResourceItemViewModel>
-	, public FGMETickableViewModelBase
+	: public FGMETickableViewModelBase
 	, public FEditorUndoClient
 	, public IGMETreeNodeViewModel
 {
