Fixed multiply-inherited TSharedFromThis.

diff --git a/Engine/Plugins/Experimental/GeometryMask/Source/GeometryMaskEditor/Private/ViewModels/GMEResourceListViewModel.h b/Engine/Plugins/Experimental/GeometryMask/Source/GeometryMaskEditor/Private/ViewModels/GMEResourceListViewModel.h
--- a/Engine/Plugins/Experimental/GeometryMask/Source/GeometryMaskEditor/Private/ViewModels/GMEResourceListViewModel.h
+++ b/Engine/Plugins/Experimental/GeometryMask/Source/GeometryMaskEditor/Private/ViewModels/GMEResourceListViewModel.h
@@ -1,4 +1,4 @@
-﻿// Copyright Epic Games, Inc. All Rights Reserved.
+// Copyright Epic Games, Inc. All Rights Reserved.
 
 #pragma once
 
@@ -9,8 +9,7 @@
 class FGMEResourceItemViewModel;
 
 class FGMEResourceListViewModel
-	: public TSharedFromThis<FGMEResourceListViewModel>
-	, public FGMEListViewModelBase
+	: public FGMEListViewModelBase
 	, public FEditorUndoClient
 	, public IGMETreeNodeViewModel
 {
diff --git a/Engine/Plugins/Experimental/GeometryMask/Source/GeometryMaskEditor/Private/ViewModels/GMEViewModelShared.h b/Engine/Plugins/Experimental/GeometryMask/Source/GeometryMaskEditor/Private/ViewModels/GMEViewModelShared.h
--- a/Engine/Plugins/Experimental/GeometryMask/Source/GeometryMaskEditor/Private/ViewModels/GMEViewModelShared.h
+++ b/Engine/Plugins/Experimental/GeometryMask/Source/GeometryMaskEditor/Private/ViewModels/GMEViewModelShared.h
@@ -1,4 +1,4 @@
-﻿// Copyright Epic Games, Inc. All Rights Reserved.
+// Copyright Epic Games, Inc. All Rights Reserved.
 
 #pragma once
 
@@ -36,8 +36,7 @@ protected:
 };
 
 class FGMEListViewModelBase
-	: public TSharedFromThis<FGMEListViewModelBase>
-	, public FGMETickableViewModelBase
+	: public FGMETickableViewModelBase
 {
 protected:
 	// Private token only allows members or friends to call MakeShared
