Motion Design Mask: Fix issue with masking not working outside editor world.
- Fix issue with masking sharing one single mask across all levels instead of each level having its own mask.

diff --git a/Engine/Plugins/Experimental/Avalanche/Source/AvalancheMask/Private/Mask2D/AvaMask2DBaseModifier.cpp b/Engine/Plugins/Experimental/Avalanche/Source/AvalancheMask/Private/Mask2D/AvaMask2DBaseModifier.cpp
--- a/Engine/Plugins/Experimental/Avalanche/Source/AvalancheMask/Private/Mask2D/AvaMask2DBaseModifier.cpp
+++ b/Engine/Plugins/Experimental/Avalanche/Source/AvalancheMask/Private/Mask2D/AvaMask2DBaseModifier.cpp
@@ -8,6 +8,7 @@
 #include "Components/PrimitiveComponent.h"
 #include "Engine/CanvasRenderTarget2D.h"
 #include "Engine/Engine.h"
+#include "Engine/Level.h"
 #include "Framework/AvaGizmoComponent.h"
 #include "GeometryMaskCanvas.h"
 #include "GeometryMaskReadComponent.h"
@@ -23,6 +24,7 @@
 #include "Materials/MaterialInstanceDynamic.h"
 #include "Modifiers/ActorModifierCoreStack.h"
 #include "Subsystems/ActorModifierCoreSubsystem.h"
+#include "Text3DComponent.h"
 
 #define LOCTEXT_NAMESPACE "AvaMask2DModifier"
 
@@ -504,8 +506,10 @@ UActorComponent* UAvaMask2DBaseModifier::FindOrAddMaskComponent(TSubclassOf<UAct
 
 bool UAvaMask2DBaseModifier::ActorSupportsMaskReadWrite(const AActor* InActor)
 {
-	// Currently only applies to primitive components
-	return InActor->FindComponentByClass<UPrimitiveComponent>() != nullptr;
+	// Only supports primitives and Text3D component.
+	// Note: Actors holding Text3D Components do eventually get Text Mesh Subobjects; but this should still return true even if these components are not yet built (the build is queued)
+	return InActor->FindComponentByClass<UText3DComponent>()
+		|| InActor->FindComponentByClass<UPrimitiveComponent>();
 }
 
 bool UAvaMask2DBaseModifier::TryResolveParentChannel()
@@ -659,17 +663,25 @@ void UAvaMask2DBaseModifier::OnMaskSetCanvas(const UGeometryMaskCanvas* InCanvas
 
 void UAvaMask2DBaseModifier::TryResolveCanvas()
 {
-	const UWorld* World = GetWorld();
-	if (!World)
+	const AActor* Actor = GetModifiedActor();
+	if (!Actor)
 	{
 		return;
 	}
-	
-	if (UGeometryMaskWorldSubsystem* MaskSubsystem = World->GetSubsystem<UGeometryMaskWorldSubsystem>())
+
+	const ULevel* Level = Actor->GetLevel();
+	if (!Level || !Level->OwningWorld)
 	{
-		UGeometryMaskCanvas* Canvas = MaskSubsystem->GetNamedCanvas(GetChannel());
-		LastResolvedCanvasName = Canvas->GetCanvasName();
-		CanvasWeak = Canvas;
+		return;
+	}
+
+	if (UGeometryMaskWorldSubsystem* MaskSubsystem = Level->OwningWorld->GetSubsystem<UGeometryMaskWorldSubsystem>())
+	{
+		if (UGeometryMaskCanvas* Canvas = MaskSubsystem->GetNamedCanvas(Level, GetChannel()))
+		{
+			LastResolvedCanvasName = Canvas->GetCanvasName();
+			CanvasWeak = Canvas;
+		}
 	}
 }
 
@@ -685,20 +697,21 @@ UTexture* UAvaMask2DBaseModifier::TryResolveCanvasTexture(AActor* InActor, FAvaM
 	{
 		return CanvasTexture;
 	}
-	
-	if (!GetWorld())
+
+	const ULevel* Level = InActor->GetLevel();
+	if (!Level || !Level->OwningWorld)
 	{
 		return nullptr;
 	}
 
 	// Try get canvas texture if already available
-	UGeometryMaskWorldSubsystem* MaskSubsystem = GetWorld()->GetSubsystem<UGeometryMaskWorldSubsystem>();
+	UGeometryMaskWorldSubsystem* MaskSubsystem = Level->OwningWorld->GetSubsystem<UGeometryMaskWorldSubsystem>();
 	if (!MaskSubsystem)
 	{
 		return nullptr;
 	}
-			
-	if (const UGeometryMaskCanvas* Canvas = MaskSubsystem->GetNamedCanvas(Channel))
+
+	if (const UGeometryMaskCanvas* Canvas = MaskSubsystem->GetNamedCanvas(Level, Channel))
 	{
 		// Cache canvas texture
 		UTexture* CanvasTexture = Canvas->GetTexture();
diff --git a/Engine/Plugins/Experimental/Avalanche/Source/AvalancheMask/Private/Mask2D/AvaMask2DReadModifier.cpp b/Engine/Plugins/Experimental/Avalanche/Source/AvalancheMask/Private/Mask2D/AvaMask2DReadModifier.cpp
--- a/Engine/Plugins/Experimental/Avalanche/Source/AvalancheMask/Private/Mask2D/AvaMask2DReadModifier.cpp
+++ b/Engine/Plugins/Experimental/Avalanche/Source/AvalancheMask/Private/Mask2D/AvaMask2DReadModifier.cpp
@@ -79,7 +79,7 @@ bool UAvaMask2DReadModifier::ApplyRead(AActor* InActor, FAvaMask2DActorData& InA
 		// Ok to continue, just nothing done for this actor
 		return true;
 	}
-	
+
 	if (const UGeometryMaskCanvas* Canvas = GetCurrentCanvas())
 	{
 		// Get MaterialCollectionHandle
@@ -109,7 +109,7 @@ bool UAvaMask2DReadModifier::ApplyRead(AActor* InActor, FAvaMask2DActorData& InA
 			ApplyParameters.MaterialParameters.ChannelAsVector = UE::GeometryMask::MaskChannelEnumToVector[ApplyParameters.MaterialParameters.Channel];
 			ApplyParameters.MaterialParameters.bInvert = bInverted;
 
-			FGeometryMaskDrawingContext DrawingContext(GeometryCanvas->GetCanvasId().World);
+			FGeometryMaskDrawingContext DrawingContext(GeometryCanvas->GetCanvasId().Level);
 
 			const FIntPoint ViewportSize = GeometryCanvas->GetResource()->GetMaxViewportSize();
 			const FVector2f ViewportPadding = FVector2f::One() - FVector2f(ViewportSize + GeometryCanvas->GetResource()->GetViewportPadding(DrawingContext)) / ViewportSize;
diff --git a/Engine/Plugins/Experimental/Avalanche/Source/AvalancheMaskEditor/Private/AvaMaskEditorMode.cpp b/Engine/Plugins/Experimental/Avalanche/Source/AvalancheMaskEditor/Private/AvaMaskEditorMode.cpp
--- a/Engine/Plugins/Experimental/Avalanche/Source/AvalancheMaskEditor/Private/AvaMaskEditorMode.cpp
+++ b/Engine/Plugins/Experimental/Avalanche/Source/AvalancheMaskEditor/Private/AvaMaskEditorMode.cpp
@@ -322,14 +322,14 @@ UGeometryMaskCanvas* UAvaMaskEditorMode::GetCanvasReferencedByActor(const AActor
 	{
 		const FGeometryMaskWriteParameters WriteComponentParameters = WriteComponent->GetParameters();
 		CanvasName = WriteComponentParameters.CanvasName;
-	}	
+	}
 	else if (const IGeometryMaskReadInterface* ReadComponent = Cast<IGeometryMaskReadInterface>(InActor->FindComponentByInterface<UGeometryMaskReadInterface>()))
 	{
 		const FGeometryMaskReadParameters ReadComponentParameters = ReadComponent->GetParameters();
 		CanvasName = ReadComponentParameters.CanvasName;
 	}
 
-	if (UGeometryMaskCanvas* Canvas = GetWorld()->GetSubsystem<UGeometryMaskWorldSubsystem>()->GetNamedCanvas(CanvasName))
+	if (UGeometryMaskCanvas* Canvas = GetWorld()->GetSubsystem<UGeometryMaskWorldSubsystem>()->GetNamedCanvas(InActor->GetLevel(), CanvasName))
 	{
 		return Canvas;
 	}
@@ -390,7 +390,7 @@ bool UAvaMaskEditorMode::AddMaskToSelected(const TArray<AActor*>& InMaskingActor
 				Subsystem->RemoveWithoutWriters();
 			}
 
-			PreviewCanvasId = FGeometryMaskCanvasId(ParentActor->GetWorld(), ChannelName);
+			PreviewCanvasId = FGeometryMaskCanvasId(ParentActor->GetLevel(), ChannelName);
 			PreviewCanvasChannel = ColorChannel;
 		}
 	}
diff --git a/Engine/Plugins/Experimental/GeometryMask/Source/GeometryMask/Private/GeometryMaskCanvas.cpp b/Engine/Plugins/Experimental/GeometryMask/Source/GeometryMask/Private/GeometryMaskCanvas.cpp
--- a/Engine/Plugins/Experimental/GeometryMask/Source/GeometryMask/Private/GeometryMaskCanvas.cpp
+++ b/Engine/Plugins/Experimental/GeometryMask/Source/GeometryMask/Private/GeometryMaskCanvas.cpp
@@ -45,6 +45,11 @@ void UGeometryMaskCanvas::PostEditChangeProperty(FPropertyChangedEvent& InProper
 }
 #endif
 
+ULevel* UGeometryMaskCanvas::GetLevel() const
+{
+	return CanvasId.Level.ResolveObjectPtr();
+}
+
 const TArray<TWeakInterfacePtr<IGeometryMaskWriteInterface>>& UGeometryMaskCanvas::GetWriters() const
 {
 	return Writers;
@@ -255,14 +260,14 @@ void UGeometryMaskCanvas::UpdateRenderParameters()
 	}
 }
 
-void UGeometryMaskCanvas::Initialize(const UWorld* InWorld, FName InCanvasName)
+void UGeometryMaskCanvas::Initialize(const ULevel* InLevel, FName InCanvasName)
 {
 	CanvasName = InCanvasName;
-	CanvasId = FGeometryMaskCanvasId(InWorld, CanvasName);
+	CanvasId = FGeometryMaskCanvasId(InLevel, CanvasName);
 }
 
 void UGeometryMaskCanvas::Update(
-	UWorld* InWorld,
+	const ULevel* InLevel,
 	FSceneView& InView)
 {
 	RemoveInvalidWriters();
diff --git a/Engine/Plugins/Experimental/GeometryMask/Source/GeometryMask/Private/GeometryMaskCanvasActor.cpp b/Engine/Plugins/Experimental/GeometryMask/Source/GeometryMask/Private/GeometryMaskCanvasActor.cpp
--- a/Engine/Plugins/Experimental/GeometryMask/Source/GeometryMask/Private/GeometryMaskCanvasActor.cpp
+++ b/Engine/Plugins/Experimental/GeometryMask/Source/GeometryMask/Private/GeometryMaskCanvasActor.cpp
@@ -64,7 +64,7 @@ bool AGeometryMaskCanvasActor::TryResolveCanvas()
 	{
 		if (UGeometryMaskWorldSubsystem* Subsystem = GetWorld()->GetSubsystem<UGeometryMaskWorldSubsystem>())
 		{
-			Canvas = Subsystem->GetNamedCanvas(CanvasName);
+			Canvas = Subsystem->GetNamedCanvas(GetLevel(), CanvasName);
 		}
 	}
 	
diff --git a/Engine/Plugins/Experimental/GeometryMask/Source/GeometryMask/Private/GeometryMaskCanvasResource.cpp b/Engine/Plugins/Experimental/GeometryMask/Source/GeometryMask/Private/GeometryMaskCanvasResource.cpp
--- a/Engine/Plugins/Experimental/GeometryMask/Source/GeometryMask/Private/GeometryMaskCanvasResource.cpp
+++ b/Engine/Plugins/Experimental/GeometryMask/Source/GeometryMask/Private/GeometryMaskCanvasResource.cpp
@@ -5,6 +5,7 @@
 #include "Algo/MaxElement.h"
 #include "Engine/Canvas.h"
 #include "Engine/CanvasRenderTarget2D.h"
+#include "Engine/Level.h"
 #include "Engine/World.h"
 #include "GeometryMaskModule.h"
 #include "GeometryMaskSettings.h"
@@ -152,11 +153,13 @@ int32 UGeometryMaskCanvasResource::Compact()
 			if (LastAvailableColorChannel != EGeometryMaskColorChannel::None)
 			{
 				FGeometryMaskCanvasId& UsedCanvasId = DependentCanvasIds[ColorChannel];
-				if (UWorld* CanvasWorld = UsedCanvasId.World.ResolveObjectPtr())
+				const ULevel* Level = UsedCanvasId.Level.ResolveObjectPtr();
+
+				if (Level && Level->OwningWorld)
 				{
-					if (UGeometryMaskWorldSubsystem* MaskSubsystem = CanvasWorld->GetSubsystem<UGeometryMaskWorldSubsystem>())
+					if (UGeometryMaskWorldSubsystem* MaskSubsystem = Level->OwningWorld->GetSubsystem<UGeometryMaskWorldSubsystem>())
 					{
-						if (UGeometryMaskCanvas* UsedCanvas = MaskSubsystem->GetNamedCanvas(UsedCanvasId.Name))
+						if (UGeometryMaskCanvas* UsedCanvas = MaskSubsystem->GetNamedCanvas(Level, UsedCanvasId.Name))
 						{
 							// Re-assign to that last available channel, freeing this one
 							UsedCanvas->AssignResource(this, LastAvailableColorChannel);						
@@ -434,21 +437,21 @@ UCanvasRenderTarget2D* UGeometryMaskCanvasResource::GetRenderTargetTexture()
 }
 
 void UGeometryMaskCanvasResource::Update(
-	UWorld* InWorld,
+	const ULevel* InLevel,
 	FSceneView& InView,
 	int32 InViewIndex)
 {
-	Draw(InWorld, InView, InViewIndex);
+	Draw(InLevel, InView, InViewIndex);
 }
 
-void UGeometryMaskCanvasResource::Draw(UWorld* InWorld, FSceneView& InView, int32 InViewIndex)
+void UGeometryMaskCanvasResource::Draw(const ULevel* InLevel, FSceneView& InView, int32 InViewIndex)
 {
-	if (!InWorld)
+	if (!InLevel || !InLevel->OwningWorld)
 	{
 		return;
 	}
 
-	UE_LOG(LogGeometryMask, VeryVerbose, TEXT("UGeometryMaskCanvasResource::Update World: %s"), *InWorld->GetName());
+	UE_LOG(LogGeometryMask, VeryVerbose, TEXT("UGeometryMaskCanvasResource::Update Level: %s"), *InLevel->GetName());
 
 	if (UCanvasRenderTarget2D* Texture = GetRenderTargetTexture())
 	{
@@ -461,15 +464,15 @@ void UGeometryMaskCanvasResource::Draw(UWorld* InWorld, FSceneView& InView, int3
 			FCanvas* NewCanvas = new FCanvas(
 				RenderTargetResource,
 				nullptr,
-				InWorld,
-				InWorld->GetFeatureLevel(),
+				InLevel->OwningWorld,
+				InLevel->OwningWorld->GetFeatureLevel(),
 				// Draw immediately so that interleaved SetVectorParameter (etc) function calls work as expected
 				FCanvas::CDM_ImmediateDrawing);
 
 			CanvasObject->Init(Texture->SizeX, Texture->SizeY, &InView, NewCanvas);
 		}
 
-		FGeometryMaskDrawingContext* DrawingContextPtr = GetDrawingContextForWorld(InWorld, InViewIndex);
+		FGeometryMaskDrawingContext* DrawingContextPtr = GetDrawingContextForLevel(InLevel, InViewIndex);
         check(DrawingContextPtr); // Should always be valid - world was already checked
 
         FGeometryMaskDrawingContext& DrawingContextRef = *DrawingContextPtr;
@@ -484,8 +487,8 @@ void UGeometryMaskCanvasResource::Draw(UWorld* InWorld, FSceneView& InView, int3
 			DrawingContextRef.ViewProjectionMatrix.M[2][2] = UE_KINDA_SMALL_NUMBER; // Prevents div by zero later
 
 			UpdateViewportSize();
-			
-			InWorld->FlushDeferredParameterCollectionInstanceUpdates();
+
+			InLevel->OwningWorld->FlushDeferredParameterCollectionInstanceUpdates();
 
 			CanvasObject->SizeX = Texture->SizeX;
 			CanvasObject->SizeY = Texture->SizeY;
@@ -536,14 +539,14 @@ void UGeometryMaskCanvasResource::Draw(UWorld* InWorld, FSceneView& InView, int3
 	}
 }
 
-FGeometryMaskDrawingContext* UGeometryMaskCanvasResource::GetDrawingContextForWorld(const UWorld* InWorld, uint8 InSceneViewIndex)
+FGeometryMaskDrawingContext* UGeometryMaskCanvasResource::GetDrawingContextForLevel(const ULevel* InLevel, uint8 InSceneViewIndex)
 {
-	if (!ensure(InWorld))
+	if (!ensure(InLevel))
 	{
 		return nullptr;
 	}
 
-	return &DrawingContextCache.FindOrAdd(FGeometryMaskDrawingContext(InWorld, InSceneViewIndex));
+	return &DrawingContextCache.FindOrAdd(FGeometryMaskDrawingContext(InLevel, InSceneViewIndex));
 }
 
 FGeometryMaskDrawingContext* UGeometryMaskCanvasResource::GetDrawingContextForCanvas(const FGeometryMaskCanvasId& InCanvasId)
@@ -553,7 +556,7 @@ FGeometryMaskDrawingContext* UGeometryMaskCanvasResource::GetDrawingContextForCa
 		return &DefaultDrawingContext;
 	}
 	
-	return GetDrawingContextForWorld(InCanvasId.World.ResolveObjectPtr(), 0);
+	return GetDrawingContextForLevel(InCanvasId.Level.ResolveObjectPtr(), 0);
 }
 
 FGeometryMaskDrawingContext* UGeometryMaskCanvasResource::GetDrawingContextForChannel(EGeometryMaskColorChannel InColorChannel)
diff --git a/Engine/Plugins/Experimental/GeometryMask/Source/GeometryMask/Private/GeometryMaskSubsystem.cpp b/Engine/Plugins/Experimental/GeometryMask/Source/GeometryMask/Private/GeometryMaskSubsystem.cpp
--- a/Engine/Plugins/Experimental/GeometryMask/Source/GeometryMask/Private/GeometryMaskSubsystem.cpp
+++ b/Engine/Plugins/Experimental/GeometryMask/Source/GeometryMask/Private/GeometryMaskSubsystem.cpp
@@ -4,6 +4,7 @@
 
 #include "Engine/Engine.h"
 #include "Engine/GameViewportClient.h"
+#include "Engine/Level.h"
 #include "Engine/LocalPlayer.h"
 #include "Engine/World.h"
 #include "GameFramework/PlayerController.h"
@@ -47,28 +48,14 @@ void UGeometryMaskSubsystem::Update(
 	}
 
 	UE_LOG(LogGeometryMask, VeryVerbose, TEXT("UGeometryMaskSubsystem::Update World: %s, Num. Views: %u"), *InWorld->GetName(), InViewFamily.Views.Num());
-	
+
 	if (UGeometryMaskWorldSubsystem* Subsystem = InWorld->GetSubsystem<UGeometryMaskWorldSubsystem>())
 	{
 		DECLARE_SCOPE_CYCLE_COUNTER(TEXT("UGeometryMaskSubsystem::Update"), STAT_GeometryMask_UpdateAll, STATGROUP_GeometryMask);
 
-		int32 ViewIndex = 0;
-		for (const FSceneView*& View : InViewFamily.Views)
+		for (const ULevel* Level : InWorld->GetLevels())
 		{
-			FSceneView* MutableSceneView = const_cast<FSceneView*>(View);
-
-			for (const TPair<FName, TObjectPtr<UGeometryMaskCanvas>>& NamedCanvas : Subsystem->NamedCanvases)
-			{
-				NamedCanvas.Value->Update(InWorld, *MutableSceneView);		
-			}
-
-			// Updates the texture resource
-			for (const TObjectPtr<UGeometryMaskCanvasResource>& Resource : CanvasResources)
-			{
-				Resource->Update(InWorld, *MutableSceneView, ViewIndex);
-			}
-
-			++ViewIndex;
+			UpdateLevel(Level, Subsystem, InViewFamily);
 		}
 	}
 }
@@ -83,6 +70,37 @@ void UGeometryMaskSubsystem::ToggleUpdate(const TOptional<bool>& bInShouldUpdate
 	}
 }
 
+void UGeometryMaskSubsystem::UpdateLevel(const ULevel* InLevel, UGeometryMaskWorldSubsystem* InWorldSubsystem, FSceneViewFamily& InViewFamily)
+{
+	if (!IsValid(InLevel) || !InWorldSubsystem)
+	{
+		return;
+	}
+
+	int32 ViewIndex = 0;
+
+	for (const FSceneView*& View : InViewFamily.Views)
+	{
+		FSceneView* MutableSceneView = const_cast<FSceneView*>(View);
+
+		if (const FGeometryMaskLevelState* LevelState = InWorldSubsystem->FindLevelState(InLevel))
+		{
+			for (const TPair<FName, TObjectPtr<UGeometryMaskCanvas>>& NamedCanvas : LevelState->NamedCanvases)
+			{
+				NamedCanvas.Value->Update(InLevel, *MutableSceneView);		
+			}
+		}
+
+		// Updates the texture resource
+		for (const TObjectPtr<UGeometryMaskCanvasResource>& Resource : CanvasResources)
+		{
+			Resource->Update(InLevel, *MutableSceneView, ViewIndex);
+		}
+
+		++ViewIndex;
+	}
+}
+
 void UGeometryMaskSubsystem::AssignResourceToCanvas(UGeometryMaskCanvas* InCanvas)
 {
 	UGeometryMaskCanvasResource* AvailableResource = nullptr;
diff --git a/Engine/Plugins/Experimental/GeometryMask/Source/GeometryMask/Private/GeometryMaskTypes.cpp b/Engine/Plugins/Experimental/GeometryMask/Source/GeometryMask/Private/GeometryMaskTypes.cpp
--- a/Engine/Plugins/Experimental/GeometryMask/Source/GeometryMask/Private/GeometryMaskTypes.cpp
+++ b/Engine/Plugins/Experimental/GeometryMask/Source/GeometryMask/Private/GeometryMaskTypes.cpp
@@ -4,6 +4,7 @@
 
 #include "Engine/CanvasRenderTarget2D.h"
 #include "Engine/Engine.h"
+#include "Engine/Level.h"
 #include "GeometryMaskSubsystem.h"
 #include "GeometryMaskWorldSubsystem.h"
 #include "UObject/UObjectThreadContext.h"
@@ -44,21 +45,21 @@ uint32 GetTypeHash(const FGeometryMaskCanvasId& InCanvasId)
 {
 	return HashCombineFast(
 		HashCombineFast(
-			GetTypeHash(InCanvasId.World),
+			GetTypeHash(InCanvasId.Level),
 			GetTypeHash(InCanvasId.Name)),
 			InCanvasId.SceneViewIndex);
 }
 
 uint32 GetTypeHash(const FGeometryMaskDrawingContext& InUpdateContext)
 {
-	return HashCombineFast(GetTypeHash(InUpdateContext.World), InUpdateContext.SceneViewIndex);
+	return HashCombineFast(GetTypeHash(InUpdateContext.Level), InUpdateContext.SceneViewIndex);
 }
 
 const FGeometryMaskCanvasId& FGeometryMaskCanvasId::None = FGeometryMaskCanvasId(EForceInit::ForceInit);
 const FName FGeometryMaskCanvasId::DefaultCanvasName = TEXT("Default");
 
-FGeometryMaskCanvasId::FGeometryMaskCanvasId(const UWorld* InWorld, const FName InName)
-	: World(InWorld)
+FGeometryMaskCanvasId::FGeometryMaskCanvasId(const ULevel* InLevel, const FName InName)
+	: Level(InLevel)
 	, Name(InName)
 {
 }
@@ -80,34 +81,39 @@ bool FGeometryMaskCanvasId::IsNone() const
 
 void FGeometryMaskCanvasId::ResetToNone()
 {
-	World = nullptr;
+	Level = nullptr;
 	SceneViewIndex = 0;
 	Name = NAME_None;
 }
 
 FString FGeometryMaskCanvasId::ToString() const
 {
+	FString LevelLabel = TEXT("(None)");
 	FString WorldLabel = TEXT("(Transient)");
 	FString WorldTypeLabel = TEXT("");
-	if (const UWorld* ResolvedWorld = World.ResolveObjectPtr())
+	if (const ULevel* ResolvedLevel = Level.ResolveObjectPtr())
 	{
-		WorldLabel = ResolvedWorld->GetName();
-		WorldTypeLabel = LexToString(ResolvedWorld->WorldType);
+		if (ResolvedLevel->OwningWorld)
+		{
+			WorldLabel = ResolvedLevel->OwningWorld->GetName();
+			WorldTypeLabel = LexToString(ResolvedLevel->OwningWorld->WorldType);
+		}
+		LevelLabel = ResolvedLevel->GetName();
 	}
 
-	return FString::Printf(TEXT("%s(%s).%s"), *WorldLabel, *WorldTypeLabel, *Name.ToString());
+	return FString::Printf(TEXT("%s(%s - %s).%s"), *LevelLabel, *WorldLabel, *WorldTypeLabel, *Name.ToString());
 }
 
-FGeometryMaskDrawingContext::FGeometryMaskDrawingContext(TObjectKey<UWorld> InWorld, const uint8 InSceneViewIndex)
-	: World(MoveTemp(InWorld))
+FGeometryMaskDrawingContext::FGeometryMaskDrawingContext(TObjectKey<ULevel> InLevel, const uint8 InSceneViewIndex)
+	: Level(MoveTemp(InLevel))
 	, SceneViewIndex(InSceneViewIndex)
 	, ViewportSize(ForceInit)
 	, ViewProjectionMatrix(ForceInit)
 {
 }
 
-FGeometryMaskDrawingContext::FGeometryMaskDrawingContext(const UWorld* InWorld, const uint8 InSceneViewIndex)
-	: World(InWorld)
+FGeometryMaskDrawingContext::FGeometryMaskDrawingContext(const ULevel* InLevel, const uint8 InSceneViewIndex)
+	: Level(InLevel)
 	, SceneViewIndex(InSceneViewIndex)
 	, ViewportSize(ForceInit)
 	, ViewProjectionMatrix(ForceInit)
@@ -115,7 +121,7 @@ FGeometryMaskDrawingContext::FGeometryMaskDrawingContext(const UWorld* InWorld,
 }
 
 FGeometryMaskDrawingContext::FGeometryMaskDrawingContext(EForceInit)
-	: World(nullptr)
+	: Level(nullptr)
 	, ViewportSize(ForceInit)
 	, ViewProjectionMatrix(ForceInit)
 {
@@ -123,7 +129,7 @@ FGeometryMaskDrawingContext::FGeometryMaskDrawingContext(EForceInit)
 
 bool FGeometryMaskDrawingContext::IsValid() const
 {
-	return World.ResolveObjectPtr() != nullptr;
+	return Level.ResolveObjectPtr() != nullptr;
 }
 
 UGeometryMaskCanvasReferenceComponentBase::~UGeometryMaskCanvasReferenceComponentBase()
@@ -193,7 +199,7 @@ bool UGeometryMaskCanvasReferenceComponentBase::TryResolveNamedCanvas(FName InCa
 
 	if (UGeometryMaskWorldSubsystem* Subsystem = GetWorld()->GetSubsystem<UGeometryMaskWorldSubsystem>())
 	{
-		Canvas = Subsystem->GetNamedCanvas(InCanvasName);
+		Canvas = Subsystem->GetNamedCanvas(GetComponentLevel(), InCanvasName);
 
 		CanvasWeak = Canvas;
 		
diff --git a/Engine/Plugins/Experimental/GeometryMask/Source/GeometryMask/Private/GeometryMaskWorldSubsystem.cpp b/Engine/Plugins/Experimental/GeometryMask/Source/GeometryMask/Private/GeometryMaskWorldSubsystem.cpp
--- a/Engine/Plugins/Experimental/GeometryMask/Source/GeometryMask/Private/GeometryMaskWorldSubsystem.cpp
+++ b/Engine/Plugins/Experimental/GeometryMask/Source/GeometryMask/Private/GeometryMaskWorldSubsystem.cpp
@@ -3,6 +3,7 @@
 #include "GeometryMaskWorldSubsystem.h"
 
 #include "Engine/Engine.h"
+#include "Engine/Level.h"
 #include "Engine/World.h"
 #include "GeometryMaskSVE.h"
 #include "SceneViewExtension.h"
@@ -20,15 +21,18 @@ void UGeometryMaskWorldSubsystem::Deinitialize()
 {
 	Super::Deinitialize();
 
-	for (const TPair<FName, TObjectPtr<UGeometryMaskCanvas>>& NamedCanvas : NamedCanvases)
+	for (const TPair<TWeakObjectPtr<const ULevel>, FGeometryMaskLevelState>& LevelState : LevelStates)
 	{
-		if (IsValid(NamedCanvas.Value))
+		for (const TPair<FName, TObjectPtr<UGeometryMaskCanvas>>& NamedCanvas : LevelState.Value.NamedCanvases)
 		{
-			NamedCanvas.Value->Free();	
+			if (IsValid(NamedCanvas.Value))
+			{
+				NamedCanvas.Value->Free();	
+			}
 		}
 	}
 
-	NamedCanvases.Empty();
+	LevelStates.Empty();
 
 	if (UGeometryMaskSubsystem* EngineSubsystem = GEngine->GetEngineSubsystem<UGeometryMaskSubsystem>())
 	{
@@ -36,27 +40,53 @@ void UGeometryMaskWorldSubsystem::Deinitialize()
 	}
 }
 
-UGeometryMaskCanvas* UGeometryMaskWorldSubsystem::GetNamedCanvas(FName InName)
+const FGeometryMaskLevelState* UGeometryMaskWorldSubsystem::FindLevelState(const ULevel* InLevel) const
 {
+	if (IsValid(InLevel))
+	{
+		return LevelStates.Find(InLevel);
+	}
+	return nullptr;
+}
+
+FGeometryMaskLevelState& UGeometryMaskWorldSubsystem::FindOrAddLevelState(const ULevel* InLevel)
+{
+	check(IsValid(InLevel));
+	return LevelStates.FindOrAdd(InLevel);
+}
+
+UGeometryMaskCanvas* UGeometryMaskWorldSubsystem::GetNamedCanvas(const ULevel* InLevel, FName InName)
+{
+	if (!IsValid(InLevel))
+	{
+		return nullptr;
+	}
+
 	UGeometryMaskSubsystem* EngineSubsystem = GEngine->GetEngineSubsystem<UGeometryMaskSubsystem>();
 	if (!ensureAlwaysMsgf(EngineSubsystem, TEXT("UGeometryMaskSubsystem not resolved.")))
 	{
 		return nullptr;
 	}
-	
+
 	if (InName.IsNone())
 	{
 		return EngineSubsystem->GetDefaultCanvas();
 	}
 
-	const FName ObjectName = MakeUniqueObjectName(this, UGeometryMaskCanvas::StaticClass(), FName(FString::Printf(TEXT("GeometryMaskCanvas_%s_"), *InName.ToString())));
-	if (TObjectPtr<UGeometryMaskCanvas>* FoundCanvas = NamedCanvases.Find(InName))
+	if (const FGeometryMaskLevelState* LevelState = FindLevelState(InLevel))
 	{
-		return *FoundCanvas;
+		if (const TObjectPtr<UGeometryMaskCanvas>* FoundCanvas = LevelState->NamedCanvases.Find(InName))
+		{
+			return *FoundCanvas;
+		}
 	}
 
-	const TObjectPtr<UGeometryMaskCanvas>& NewCanvas = NamedCanvases.Emplace(InName, NewObject<UGeometryMaskCanvas>(this, ObjectName));
-	NewCanvas->Initialize(GetWorld(), InName);
+	const FName ObjectName = MakeUniqueObjectName(this, UGeometryMaskCanvas::StaticClass(), FName(FString::Printf(TEXT("GeometryMaskCanvas_%s_"), *InName.ToString())));
+
+	FGeometryMaskLevelState& LevelState = FindOrAddLevelState(InLevel);
+	const TObjectPtr<UGeometryMaskCanvas>& NewCanvas = LevelState.NamedCanvases.Emplace(InName, NewObject<UGeometryMaskCanvas>(this, ObjectName));
+
+	NewCanvas->Initialize(InLevel, InName);
 	EngineSubsystem->AssignResourceToCanvas(NewCanvas);
 
 	NewCanvas->OnActivated().BindUObject(this, &UGeometryMaskWorldSubsystem::OnCanvasActivated, NewCanvas.Get());
@@ -67,10 +97,13 @@ UGeometryMaskCanvas* UGeometryMaskWorldSubsystem::GetNamedCanvas(FName InName)
 	return NewCanvas;
 }
 
-TArray<FName> UGeometryMaskWorldSubsystem::GetCanvasNames()
+TArray<FName> UGeometryMaskWorldSubsystem::GetCanvasNames(const ULevel* InLevel)
 {
 	TArray<FName> CanvasNames;
-	NamedCanvases.GenerateKeyArray(CanvasNames);
+	if (const FGeometryMaskLevelState* LevelState = FindLevelState(InLevel))
+	{
+		LevelState->NamedCanvases.GenerateKeyArray(CanvasNames);
+	}
 	return CanvasNames;
 }
 
@@ -79,28 +112,33 @@ int32 UGeometryMaskWorldSubsystem::RemoveWithoutWriters()
 	int32 NumRemoved = 0;
 	
 	TMap<FName, TObjectPtr<UGeometryMaskCanvas>> UsedCanvases;
-	UsedCanvases.Reserve(NamedCanvases.Num());
+	UsedCanvases.Reserve(LevelStates.Num());
 
-	for (const TPair<FName, TObjectPtr<UGeometryMaskCanvas>>& NamedCanvas : NamedCanvases)
+	for (decltype(LevelStates)::TIterator LevelStateIter(LevelStates); LevelStateIter; ++LevelStateIter)
 	{
-		if (!IsValid(NamedCanvas.Value))
-		{
-			continue;
-		}
-		
-		if (!NamedCanvas.Value->GetWriters().IsEmpty())
+		FGeometryMaskLevelState& LevelState = LevelStateIter.Value();
+
+		for (decltype(LevelState.NamedCanvases)::TIterator CanvasIter(LevelState.NamedCanvases); CanvasIter; ++CanvasIter)
 		{
-			UsedCanvases.Emplace(NamedCanvas.Key, NamedCanvas.Value);
+			TObjectPtr<UGeometryMaskCanvas> NamedCanvas = CanvasIter.Value();
+			if (IsValid(NamedCanvas))
+			{
+				if (NamedCanvas->GetWriters().IsEmpty())
+				{
+					OnGeometryMaskCanvasDestroyed().Broadcast(NamedCanvas->GetCanvasId());
+					NamedCanvas->FreeResource();
+				}
+
+				CanvasIter.RemoveCurrent();
+				++NumRemoved;
+			}
 		}
-		else
+
+		if (LevelState.NamedCanvases.IsEmpty())
 		{
-			OnGeometryMaskCanvasDestroyed().Broadcast(NamedCanvas.Value->GetCanvasId());
-			NamedCanvas.Value->FreeResource();
-			++NumRemoved;
+			LevelStateIter.RemoveCurrent();
 		}
 	}
-	
-	NamedCanvases = UsedCanvases;
 
 	return NumRemoved;
 }
diff --git a/Engine/Plugins/Experimental/GeometryMask/Source/GeometryMask/Public/GeometryMaskCanvas.h b/Engine/Plugins/Experimental/GeometryMask/Source/GeometryMask/Public/GeometryMaskCanvas.h
--- a/Engine/Plugins/Experimental/GeometryMask/Source/GeometryMask/Public/GeometryMaskCanvas.h
+++ b/Engine/Plugins/Experimental/GeometryMask/Source/GeometryMask/Public/GeometryMaskCanvas.h
@@ -36,6 +36,8 @@ public:
 	virtual void PostEditChangeProperty(FPropertyChangedEvent& InPropertyChangedEvent) override;
 #endif
 
+	ULevel* GetLevel() const;
+
 	/** Returns all writers. */
     const TArray<TWeakInterfacePtr<IGeometryMaskWriteInterface>>& GetWriters() const;
 
@@ -105,10 +107,10 @@ public:
 	FOnGeometryMaskCanvasDeactivated& OnDeactivated() { return OnDeactivatedDelegate; }
 
 	/** Setup with identifying info. */
-	void Initialize(const UWorld* InWorld, FName InCanvasName);
+	void Initialize(const ULevel* InLevel, FName InCanvasName);
 	
 	/** Updates the canvas, intended to be called every frame. */
-	void Update(UWorld* InWorld, FSceneView& InView);
+	void Update(const ULevel* InLevel, FSceneView& InView);
 	
 	const UGeometryMaskCanvasResource* GetResource() const { return CanvasResource; }
 	void AssignResource(UGeometryMaskCanvasResource* InResource, EGeometryMaskColorChannel InColorChannel);
diff --git a/Engine/Plugins/Experimental/GeometryMask/Source/GeometryMask/Public/GeometryMaskCanvasResource.h b/Engine/Plugins/Experimental/GeometryMask/Source/GeometryMask/Public/GeometryMaskCanvasResource.h
--- a/Engine/Plugins/Experimental/GeometryMask/Source/GeometryMask/Public/GeometryMaskCanvasResource.h
+++ b/Engine/Plugins/Experimental/GeometryMask/Source/GeometryMask/Public/GeometryMaskCanvasResource.h
@@ -71,13 +71,13 @@ public:
 	void ResetRenderParameters(EGeometryMaskColorChannel InColorChannel);
 
 	/** Updates the canvas, intended to be called every frame. */
-	void Update(UWorld* InWorld, FSceneView& InView, int32 InViewIndex = 0);
+	void Update(const ULevel* InLevel, FSceneView& InView, int32 InViewIndex = 0);
 
 private:
 	/** Draws all writers to the canvas. */
-	void Draw(UWorld* InWorld, FSceneView& InView, int32 InViewIndex = 0);
+	void Draw(const ULevel* InLevel, FSceneView& InView, int32 InViewIndex = 0);
 
-	FGeometryMaskDrawingContext* GetDrawingContextForWorld(const UWorld* InWorld, uint8 InSceneViewIndex);
+	FGeometryMaskDrawingContext* GetDrawingContextForLevel(const ULevel* InLevel, uint8 InSceneViewIndex);
 	FGeometryMaskDrawingContext* GetDrawingContextForCanvas(const FGeometryMaskCanvasId& InCanvasId);
 	FGeometryMaskDrawingContext* GetDrawingContextForChannel(EGeometryMaskColorChannel InColorChannel);
 
diff --git a/Engine/Plugins/Experimental/GeometryMask/Source/GeometryMask/Public/GeometryMaskSubsystem.h b/Engine/Plugins/Experimental/GeometryMask/Source/GeometryMask/Public/GeometryMaskSubsystem.h
--- a/Engine/Plugins/Experimental/GeometryMask/Source/GeometryMask/Public/GeometryMaskSubsystem.h
+++ b/Engine/Plugins/Experimental/GeometryMask/Source/GeometryMask/Public/GeometryMaskSubsystem.h
@@ -11,6 +11,8 @@
 
 #include "GeometryMaskSubsystem.generated.h"
 
+class UGeometryMaskWorldSubsystem;
+
 using FOnGeometryMaskResourceCreated = TMulticastDelegate<void(const UGeometryMaskCanvasResource*)>;
 using FOnGeometryMaskResourceDestroyed = TMulticastDelegate<void(const UGeometryMaskCanvasResource*)>;
 
@@ -41,7 +43,9 @@ public:
 	/** Called when a canvas resource is destroyed. */
 	FOnGeometryMaskResourceDestroyed& OnGeometryMaskResourceDestroyed() { return OnGeometryMaskResourceDestroyedDelegate; }
 
-private:	
+private:
+	void UpdateLevel(const ULevel* InLevel, UGeometryMaskWorldSubsystem* InWorldSubsystem, FSceneViewFamily& InViewFamily);
+
 	/** Find and assign the next available resource to the given canvas. */
 	void AssignResourceToCanvas(UGeometryMaskCanvas* InCanvas);
 
diff --git a/Engine/Plugins/Experimental/GeometryMask/Source/GeometryMask/Public/GeometryMaskTypes.h b/Engine/Plugins/Experimental/GeometryMask/Source/GeometryMask/Public/GeometryMaskTypes.h
--- a/Engine/Plugins/Experimental/GeometryMask/Source/GeometryMask/Public/GeometryMaskTypes.h
+++ b/Engine/Plugins/Experimental/GeometryMask/Source/GeometryMask/Public/GeometryMaskTypes.h
@@ -78,7 +78,7 @@ struct GEOMETRYMASK_API FGeometryMaskCanvasId
 	GENERATED_BODY()
 
 public:
-	TObjectKey<UWorld> World;
+	TObjectKey<ULevel> Level;
 
 	UPROPERTY()
 	uint8 SceneViewIndex = 0;
@@ -92,7 +92,7 @@ public:
 
 public:
 	FGeometryMaskCanvasId() = default;
-	explicit FGeometryMaskCanvasId(const UWorld* InWorld, const FName InName);
+	explicit FGeometryMaskCanvasId(const ULevel* InLevel, const FName InName);
 	explicit FGeometryMaskCanvasId(EForceInit);
 	
 	bool IsDefault() const;
@@ -104,7 +104,7 @@ public:
 public:
 	bool operator==(const FGeometryMaskCanvasId& InOther) const
 	{
-		return World == InOther.World && Name.IsEqual(InOther.Name) && SceneViewIndex == InOther.SceneViewIndex;
+		return Level == InOther.Level && Name.IsEqual(InOther.Name) && SceneViewIndex == InOther.SceneViewIndex;
 	}
 	
 	bool operator!=(const FGeometryMaskCanvasId& InOther) const
@@ -122,7 +122,7 @@ struct GEOMETRYMASK_API FGeometryMaskDrawingContext
 	GENERATED_BODY()
 
 public:
-	TObjectKey<UWorld> World;
+	TObjectKey<ULevel> Level;
 	
 	UPROPERTY()
 	uint8 SceneViewIndex = 0;
@@ -136,8 +136,8 @@ public:
 
 public:
 	FGeometryMaskDrawingContext() = default;
-	explicit FGeometryMaskDrawingContext(TObjectKey<UWorld> InWorld, const uint8 InSceneViewIndex = 0);
-	explicit FGeometryMaskDrawingContext(const UWorld* InWorld, const uint8 InSceneViewIndex = 0);
+	explicit FGeometryMaskDrawingContext(TObjectKey<ULevel> InLevel, const uint8 InSceneViewIndex = 0);
+	explicit FGeometryMaskDrawingContext(const ULevel* InLevel, const uint8 InSceneViewIndex = 0);
 	explicit FGeometryMaskDrawingContext(EForceInit);
 	
 	bool IsValid() const;
@@ -145,7 +145,7 @@ public:
 public:
 	bool operator==(const FGeometryMaskDrawingContext& InOther) const
 	{
-		return World == InOther.World && SceneViewIndex == InOther.SceneViewIndex;
+		return Level == InOther.Level && SceneViewIndex == InOther.SceneViewIndex;
 	}
 
 	bool operator!=(const FGeometryMaskDrawingContext& InOther) const
diff --git a/Engine/Plugins/Experimental/GeometryMask/Source/GeometryMask/Public/GeometryMaskWorldSubsystem.h b/Engine/Plugins/Experimental/GeometryMask/Source/GeometryMask/Public/GeometryMaskWorldSubsystem.h
--- a/Engine/Plugins/Experimental/GeometryMask/Source/GeometryMask/Public/GeometryMaskWorldSubsystem.h
+++ b/Engine/Plugins/Experimental/GeometryMask/Source/GeometryMask/Public/GeometryMaskWorldSubsystem.h
@@ -9,11 +9,19 @@
 #include "GeometryMaskWorldSubsystem.generated.h"
 
 class FGeometryMaskSceneViewExtension;
+class ULevel;
 
 using FOnGeometryMaskCanvasCreated = TMulticastDelegate<void(const UGeometryMaskCanvas*)>;
 using FOnGeometryMaskCanvasDestroyed = TMulticastDelegate<void(const FGeometryMaskCanvasId&)>;
 
+USTRUCT()
+struct FGeometryMaskLevelState
+{
+	GENERATED_BODY()
 
+	UPROPERTY()
+	TMap<FName, TObjectPtr<UGeometryMaskCanvas>> NamedCanvases;
+};
 
 /** Updates the canvases. */
 UCLASS()
@@ -24,12 +32,10 @@ class GEOMETRYMASK_API UGeometryMaskWorldSubsystem
 
 public:
 	/** Retrieves a Canvas, uniquely identified by this world and the canvas name. */
-	UFUNCTION(BlueprintCallable, Category = "Canvas")
-	UGeometryMaskCanvas* GetNamedCanvas(FName InName);
+	UGeometryMaskCanvas* GetNamedCanvas(const ULevel* InLevel, FName InName);
 
 	/** Returns all registered canvas names for this world. */
-	UFUNCTION(BlueprintCallable, Category = "Canvas")
-	TArray<FName> GetCanvasNames();
+	TArray<FName> GetCanvasNames(const ULevel* InLevel);
 	
 	/** Remove all canvases without any Readers or Writers. Return the number of canvases removed. */
 	int32 RemoveWithoutWriters();
@@ -47,17 +53,21 @@ protected:
 	// ~End USubsystem
 
 private:
+	const FGeometryMaskLevelState* FindLevelState(const ULevel* InLevel) const;
+
+	FGeometryMaskLevelState& FindOrAddLevelState(const ULevel* InLevel);
+
 	void OnCanvasActivated(UGeometryMaskCanvas* InCanvas);
 	void OnCanvasDeactivated(UGeometryMaskCanvas* InCanvas);
 
 private:
 	friend class UGeometryMaskSubsystem;
-	
+
 	TSharedPtr<FGeometryMaskSceneViewExtension> GeometryMaskSceneViewExtension;
 
 	FOnGeometryMaskCanvasCreated OnGeometryMaskCanvasCreatedDelegate;
 	FOnGeometryMaskCanvasDestroyed OnGeometryMaskCanvasDestroyedDelegate;
 
 	UPROPERTY()
-	TMap<FName, TObjectPtr<UGeometryMaskCanvas>> NamedCanvases;
+	TMap<TWeakObjectPtr<const ULevel>, FGeometryMaskLevelState> LevelStates;
 };
diff --git a/Engine/Plugins/Experimental/GeometryMask/Source/GeometryMaskEditor/Private/GeometryMaskEditorModule.cpp b/Engine/Plugins/Experimental/GeometryMask/Source/GeometryMaskEditor/Private/GeometryMaskEditorModule.cpp
--- a/Engine/Plugins/Experimental/GeometryMask/Source/GeometryMaskEditor/Private/GeometryMaskEditorModule.cpp
+++ b/Engine/Plugins/Experimental/GeometryMask/Source/GeometryMaskEditor/Private/GeometryMaskEditorModule.cpp
@@ -146,18 +146,21 @@ void FGeometryMaskEditorModule::ExecuteFlush(const TArray<FString>& InArgs)
 {
 	for (const FWorldContext& WorldContext : GEngine->GetWorldContexts())
 	{
-		if (UWorld* World = WorldContext.World())
+		if (const UWorld* World = WorldContext.World())
 		{
 			if (UGeometryMaskWorldSubsystem* GeometryMaskSubsystem = World->GetSubsystem<UGeometryMaskWorldSubsystem>())
 			{
 				const int32 RemovedCanvasNum = GeometryMaskSubsystem->RemoveWithoutWriters();
-				const int32 ActiveCanvasNum = GeometryMaskSubsystem->GetCanvasNames().Num();
+
+				int32 ActiveCanvasNum = 0;
+				for (const ULevel* Level : World->GetLevels())
+				{
+					ActiveCanvasNum += GeometryMaskSubsystem->GetCanvasNames(Level).Num();
+				}
 				UE_LOG(LogGeometryMaskEditor, Display, TEXT("%u canvas's removed because they had no writers - %u canvas's remaining."), RemovedCanvasNum, ActiveCanvasNum);
 			}
 		}
 	}
-	
-
 }
 
 #undef LOCTEXT_NAMESPACE
diff --git a/Engine/Plugins/Experimental/GeometryMask/Source/GeometryMaskEditor/Private/ViewModels/GMECanvasItemViewModel.cpp b/Engine/Plugins/Experimental/GeometryMask/Source/GeometryMaskEditor/Private/ViewModels/GMECanvasItemViewModel.cpp
--- a/Engine/Plugins/Experimental/GeometryMask/Source/GeometryMaskEditor/Private/ViewModels/GMECanvasItemViewModel.cpp
+++ b/Engine/Plugins/Experimental/GeometryMask/Source/GeometryMaskEditor/Private/ViewModels/GMECanvasItemViewModel.cpp
@@ -3,6 +3,7 @@
 #include "GMECanvasItemViewModel.h"
 
 #include "Engine/CanvasRenderTarget2D.h"
+#include "Engine/Level.h"
 #include "Engine/Texture.h"
 #include "GeometryMaskCanvas.h"
 
@@ -22,7 +23,7 @@ FGMECanvasItemViewModel::FGMECanvasItemViewModel(
 {
 	if (const UGeometryMaskCanvas* Canvas = InCanvas.Get())
 	{
-		CanvasId = FGeometryMaskCanvasId(Canvas->GetWorld(), Canvas->GetCanvasName());		
+		CanvasId = Canvas->GetCanvasId();		
 		ColorChannel = Canvas->GetColorChannel();
 		CanvasTextureWeak = Canvas->GetTexture();
 		KnownWriterCount = Canvas->GetWriters().Num();
@@ -59,12 +60,17 @@ void FGMECanvasItemViewModel::UpdateInfoText()
 
 		FString WorldTypeLabel = TEXT("");
 		FString WorldLabel = TEXT("(None)");
-		if (UWorld* CanvasWorld = Id.World.ResolveObjectPtr())
+		FString LevelLabel = TEXT("(None)");
+		if (ULevel* CanvasLevel = Id.Level.ResolveObjectPtr())
 		{
-			WorldTypeLabel = LexToString(CanvasWorld->WorldType);
-			WorldLabel = CanvasWorld->GetName();
+			if (CanvasLevel->OwningWorld)
+			{
+				WorldLabel = CanvasLevel->OwningWorld->GetName();
+				WorldTypeLabel = LexToString(CanvasLevel->OwningWorld->WorldType);
+			}
+			LevelLabel = CanvasLevel->GetName();
 		}
-		
+
 		static const FString DefaultCanvasLabel = FGeometryMaskCanvasId::DefaultCanvasName.ToString();
 		FString CanvasLabel = DefaultCanvasLabel;
 		if (!Id.IsDefault())
@@ -73,10 +79,11 @@ void FGMECanvasItemViewModel::UpdateInfoText()
 		}
 
 		InfoText = FText::FromString(
-			FString::Printf(TEXT("%-s\n%-12s: %s\n%-12s: %s\n%-12s: %s\n%-12s: %u"),
-				*FString::Printf(TEXT("%s.%s"), *WorldLabel, *CanvasLabel),
+			FString::Printf(TEXT("%-s\n%-12s: %s\n%-12s:  %s\n%-12s: %s\n%-12s: %s\n%-12s: %u"),
+				*FString::Printf(TEXT("%s.%s"), *LevelLabel, *CanvasLabel),
 				TEXT("World Type"),	*WorldTypeLabel,
 				TEXT("World"), *WorldLabel,
+				TEXT("Level"), *LevelLabel,
 				TEXT("Name"), *CanvasLabel,
 				TEXT("Num. Writers"), Canvas->GetNumWriters()));
 	}
diff --git a/Engine/Plugins/Experimental/GeometryMask/Source/GeometryMaskEditor/Private/ViewModels/GMECanvasListViewModel.cpp b/Engine/Plugins/Experimental/GeometryMask/Source/GeometryMaskEditor/Private/ViewModels/GMECanvasListViewModel.cpp
--- a/Engine/Plugins/Experimental/GeometryMask/Source/GeometryMaskEditor/Private/ViewModels/GMECanvasListViewModel.cpp
+++ b/Engine/Plugins/Experimental/GeometryMask/Source/GeometryMaskEditor/Private/ViewModels/GMECanvasListViewModel.cpp
@@ -52,20 +52,31 @@ bool FGMECanvasListViewModel::RefreshItems()
 
 		for (const TObjectKey<UWorld>& WorldKey : LoadedWorlds)
 		{
-			if (UWorld* World = WorldKey.ResolveObjectPtr())
+			UWorld* World = WorldKey.ResolveObjectPtr();
+			if (!World)
 			{
-				if (UGeometryMaskWorldSubsystem* Subsystem = World->GetSubsystem<UGeometryMaskWorldSubsystem>())
+				continue;
+			}
+
+			UGeometryMaskWorldSubsystem* Subsystem = World->GetSubsystem<UGeometryMaskWorldSubsystem>();
+			if (!Subsystem)
+			{
+				continue;
+			}
+
+			for (const ULevel* Level : World->GetLevels())
+			{
+				TArray<FName> CanvasNames = Subsystem->GetCanvasNames(Level);
+				for (const FName CanvasName : CanvasNames)
 				{
-					TArray<FName> CanvasNames = Subsystem->GetCanvasNames();
-					for (const FName CanvasName : CanvasNames)
+					if (UGeometryMaskCanvas* Canvas = Subsystem->GetNamedCanvas(Level, CanvasName))
 					{
-						UGeometryMaskCanvas* Canvas = Subsystem->GetNamedCanvas(CanvasName);
 						CanvasItems.Add(FGMECanvasItemViewModel::Create(Canvas));
 					}
-
-					OnCanvasCreatedHandles.Emplace(World, Subsystem->OnGeometryMaskCanvasCreated().AddRaw(this, &FGMECanvasListViewModel::OnCanvasCreated));
-					OnCanvasDestroyedHandles.Emplace(World, Subsystem->OnGeometryMaskCanvasDestroyed().AddRaw(this, &FGMECanvasListViewModel::OnCanvasDestroyed));
 				}
+
+				OnCanvasCreatedHandles.Emplace(World, Subsystem->OnGeometryMaskCanvasCreated().AddRaw(this, &FGMECanvasListViewModel::OnCanvasCreated));
+				OnCanvasDestroyedHandles.Emplace(World, Subsystem->OnGeometryMaskCanvasDestroyed().AddRaw(this, &FGMECanvasListViewModel::OnCanvasDestroyed));
 			}
 		}
 	}
diff --git a/Engine/Plugins/Experimental/GeometryMask/Source/GeometryMaskEditor/Private/Widgets/SGeometryMaskCanvasPreview.cpp b/Engine/Plugins/Experimental/GeometryMask/Source/GeometryMaskEditor/Private/Widgets/SGeometryMaskCanvasPreview.cpp
--- a/Engine/Plugins/Experimental/GeometryMask/Source/GeometryMaskEditor/Private/Widgets/SGeometryMaskCanvasPreview.cpp
+++ b/Engine/Plugins/Experimental/GeometryMask/Source/GeometryMaskEditor/Private/Widgets/SGeometryMaskCanvasPreview.cpp
@@ -5,6 +5,7 @@
 #include "EditorSupportDelegates.h"
 #include "Engine/CanvasRenderTarget2D.h"
 #include "Engine/Engine.h"
+#include "Engine/Level.h"
 #include "Engine/Texture.h"
 #include "GeometryMaskCanvasResource.h"
 #include "GeometryMaskEditorLog.h"
@@ -247,12 +248,15 @@ bool SGeometryMaskCanvasPreview::TryResolveCanvas()
 		return false;
 	}
 
-	if (UWorld* CanvasWorld = CanvasId.Get().World.ResolveObjectPtr())
+	if (ULevel* CanvasLevel = CanvasId.Get().Level.ResolveObjectPtr())
 	{
-		if (UGeometryMaskWorldSubsystem* Subsystem = CanvasWorld->GetSubsystem<UGeometryMaskWorldSubsystem>())
+		if (CanvasLevel->OwningWorld)
 		{
-			CanvasWeak = Subsystem->GetNamedCanvas(GetCanvasName());
-			UpdateBrush(CanvasWeak.Get(), nullptr);
+			if (UGeometryMaskWorldSubsystem* Subsystem = CanvasLevel->OwningWorld->GetSubsystem<UGeometryMaskWorldSubsystem>())
+			{
+				CanvasWeak = Subsystem->GetNamedCanvas(CanvasLevel, GetCanvasName());
+				UpdateBrush(CanvasWeak.Get(), nullptr);
+			}
 		}
 	}
 
@@ -277,7 +281,7 @@ void SGeometryMaskCanvasPreview::UpdateBrush(const UGeometryMaskCanvas* InCanvas
 			InTexture = CanvasTexture;
 		}
 
-		FGeometryMaskDrawingContext DrawingContext(InCanvas->GetCanvasId().World);
+		FGeometryMaskDrawingContext DrawingContext(InCanvas->GetCanvasId().Level);
 
 		FVector4f Padding(ForceInitToZero);
 		if (const UGeometryMaskCanvasResource* CanvasResource = InCanvas->GetResource())
