Geometry Mask: Fix delegates still being tied to world, when they should be tied to level.

diff --git a/Engine/Plugins/Experimental/GeometryMask/Source/GeometryMaskEditor/Private/ViewModels/GMECanvasListViewModel.cpp b/Engine/Plugins/Experimental/GeometryMask/Source/GeometryMaskEditor/Private/ViewModels/GMECanvasListViewModel.cpp
--- a/Engine/Plugins/Experimental/GeometryMask/Source/GeometryMaskEditor/Private/ViewModels/GMECanvasListViewModel.cpp
+++ b/Engine/Plugins/Experimental/GeometryMask/Source/GeometryMaskEditor/Private/ViewModels/GMECanvasListViewModel.cpp
@@ -23,14 +23,19 @@ FGMECanvasListViewModel::~FGMECanvasListViewModel()
 		{
 			if (UGeometryMaskWorldSubsystem* Subsystem = World->GetSubsystem<UGeometryMaskWorldSubsystem>())
 			{
-				if (FDelegateHandle* DelegateHandle = OnCanvasCreatedHandles.Find(World))
+				for (ULevel* Level : World->GetLevels())
 				{
-					Subsystem->OnGeometryMaskCanvasCreated().Remove(*DelegateHandle);
-				}
+					if (FDelegateHandle* DelegateHandle = OnCanvasCreatedHandles.Find(Level))
+					{
+						Subsystem->OnGeometryMaskCanvasCreated().Remove(*DelegateHandle);
+						DelegateHandle->Reset();
+					}
 
-				if (FDelegateHandle* DelegateHandle = OnCanvasDestroyedHandles.Find(World))
-				{
-					Subsystem->OnGeometryMaskCanvasDestroyed().Remove(*DelegateHandle);
+					if (FDelegateHandle* DelegateHandle = OnCanvasDestroyedHandles.Find(Level))
+					{
+						Subsystem->OnGeometryMaskCanvasDestroyed().Remove(*DelegateHandle);
+						DelegateHandle->Reset();
+					}
 				}
 			}
 		}
@@ -75,8 +80,8 @@ bool FGMECanvasListViewModel::RefreshItems()
 					}
 				}
 
-				OnCanvasCreatedHandles.Emplace(World, Subsystem->OnGeometryMaskCanvasCreated().AddRaw(this, &FGMECanvasListViewModel::OnCanvasCreated));
-				OnCanvasDestroyedHandles.Emplace(World, Subsystem->OnGeometryMaskCanvasDestroyed().AddRaw(this, &FGMECanvasListViewModel::OnCanvasDestroyed));
+				OnCanvasCreatedHandles.Emplace(Level, Subsystem->OnGeometryMaskCanvasCreated().AddSP(this, &FGMECanvasListViewModel::OnCanvasCreated));
+				OnCanvasDestroyedHandles.Emplace(Level, Subsystem->OnGeometryMaskCanvasDestroyed().AddSP(this, &FGMECanvasListViewModel::OnCanvasDestroyed));
 			}
 		}
 	}
@@ -96,8 +101,11 @@ void FGMECanvasListViewModel::OnPostWorldInit(UWorld* InWorld, const UWorld::Ini
 			{
 				if (UGeometryMaskWorldSubsystem* Subsystem = World->GetSubsystem<UGeometryMaskWorldSubsystem>())
 				{
-					OnCanvasCreatedHandles.Emplace(World, Subsystem->OnGeometryMaskCanvasCreated().AddRaw(this, &FGMECanvasListViewModel::OnCanvasCreated));
-					OnCanvasDestroyedHandles.Emplace(World, Subsystem->OnGeometryMaskCanvasDestroyed().AddRaw(this, &FGMECanvasListViewModel::OnCanvasDestroyed));
+					for (const ULevel* Level : World->GetLevels())
+					{
+						OnCanvasCreatedHandles.Emplace(Level, Subsystem->OnGeometryMaskCanvasCreated().AddSP(this, &FGMECanvasListViewModel::OnCanvasCreated));
+						OnCanvasDestroyedHandles.Emplace(Level, Subsystem->OnGeometryMaskCanvasDestroyed().AddSP(this, &FGMECanvasListViewModel::OnCanvasDestroyed));
+					}
 				}
 			}
 		}
@@ -114,14 +122,17 @@ void FGMECanvasListViewModel::OnPreWorldDestroyed(UWorld* InWorld)
 		{
 			if (UGeometryMaskWorldSubsystem* Subsystem = World->GetSubsystem<UGeometryMaskWorldSubsystem>())
 			{
-				if (FDelegateHandle* DelegateHandle = OnCanvasCreatedHandles.Find(World))
+				for (const ULevel* Level : World->GetLevels())
 				{
-					Subsystem->OnGeometryMaskCanvasCreated().Remove(*DelegateHandle);
-				}
+					if (FDelegateHandle* DelegateHandle = OnCanvasCreatedHandles.Find(Level))
+					{
+						Subsystem->OnGeometryMaskCanvasCreated().Remove(*DelegateHandle);
+					}
 
-				if (FDelegateHandle* DelegateHandle = OnCanvasDestroyedHandles.Find(World))
-				{
-					Subsystem->OnGeometryMaskCanvasDestroyed().Remove(*DelegateHandle);
+					if (FDelegateHandle* DelegateHandle = OnCanvasDestroyedHandles.Find(Level))
+					{
+						Subsystem->OnGeometryMaskCanvasDestroyed().Remove(*DelegateHandle);
+					}
 				}
 			}
 		}
diff --git a/Engine/Plugins/Experimental/GeometryMask/Source/GeometryMaskEditor/Private/ViewModels/GMECanvasListViewModel.h b/Engine/Plugins/Experimental/GeometryMask/Source/GeometryMaskEditor/Private/ViewModels/GMECanvasListViewModel.h
--- a/Engine/Plugins/Experimental/GeometryMask/Source/GeometryMaskEditor/Private/ViewModels/GMECanvasListViewModel.h
+++ b/Engine/Plugins/Experimental/GeometryMask/Source/GeometryMaskEditor/Private/ViewModels/GMECanvasListViewModel.h
@@ -41,11 +41,11 @@ private:
 	void OnCanvasDestroyed(const FGeometryMaskCanvasId& InGeometryMaskCanvasId);
 
 private:
-	TMap<TObjectKey<UWorld>, FDelegateHandle> OnCanvasCreatedHandles;
-	TMap<TObjectKey<UWorld>, FDelegateHandle> OnCanvasDestroyedHandles;
+	TMap<TObjectKey<ULevel>, FDelegateHandle> OnCanvasCreatedHandles;
+	TMap<TObjectKey<ULevel>, FDelegateHandle> OnCanvasDestroyedHandles;
 
 	/** Cached canvas names for comparison/refresh. */
-	TMap<TObjectKey<UWorld>, TArray<FName>> LastCanvasNames;
+	TMap<TObjectKey<ULevel>, TArray<FName>> LastCanvasNames;
 
 	/** Canvas ViewModels */
 	TArray<TSharedPtr<FGMECanvasItemViewModel>> CanvasItems;
