GeometryMask: Fixed invalid use of FStringView::CopyString

CopyString requires a non-null buffer to copy into and does not terminate the copied string.

diff --git a/Engine/Plugins/Experimental/GeometryMask/Source/GeometryMask/Private/GeometryMaskCanvasResource.cpp b/Engine/Plugins/Experimental/GeometryMask/Source/GeometryMask/Private/GeometryMaskCanvasResource.cpp
--- a/Engine/Plugins/Experimental/GeometryMask/Source/GeometryMask/Private/GeometryMaskCanvasResource.cpp
+++ b/Engine/Plugins/Experimental/GeometryMask/Source/GeometryMask/Private/GeometryMaskCanvasResource.cpp
@@ -289,12 +289,10 @@ void UGeometryMaskCanvasResource::UpdateRenderParameters(
 {
 	const int32 ChannelIdx = static_cast<int32>(InColorChannel);
 
-	auto LogOutOfBounds = [InColorChannel](){
+	auto LogOutOfBounds = [InColorChannel]
+	{
 		const FStringView ChannelStringView = UE::GeometryMask::ChannelToString(InColorChannel);
-		TCHAR* ChannelString = nullptr;
-		ChannelStringView.CopyString(ChannelString, ChannelStringView.Len());
-		
-		UE_LOG(LogGeometryMask, Error, TEXT("ColorChannel wasn't valid for setting shader parameters (was '%s', expected R, G, B or A"), ChannelString);
+		UE_LOG(LogGeometryMask, Error, TEXT("ColorChannel wasn't valid for setting shader parameters (was '%.*s', expected R, G, B or A"), ChannelStringView.Len(), ChannelStringView.GetData());
 	};
 
 	FGeometryMaskPostProcessParameters_Blur BlurParameters = PostProcess_Blur->GetParameters();
