[Motion Design] Setting default verbosity to "Log" instead of "VeryVerbose".

diff --git a/Engine/Plugins/Experimental/Avalanche/Source/AvalancheMask/Private/AvaMaskLog.h b/Engine/Plugins/Experimental/Avalanche/Source/AvalancheMask/Private/AvaMaskLog.h
--- a/Engine/Plugins/Experimental/Avalanche/Source/AvalancheMask/Private/AvaMaskLog.h
+++ b/Engine/Plugins/Experimental/Avalanche/Source/AvalancheMask/Private/AvaMaskLog.h
@@ -4,4 +4,4 @@
 
 #include "Logging/LogMacros.h"
 
-DECLARE_LOG_CATEGORY_EXTERN(LogAvaMask, VeryVerbose, All);
+DECLARE_LOG_CATEGORY_EXTERN(LogAvaMask, Log, All);
diff --git a/Engine/Plugins/Experimental/Avalanche/Source/AvalancheMaskEditor/Private/AvaMaskEditorLog.h b/Engine/Plugins/Experimental/Avalanche/Source/AvalancheMaskEditor/Private/AvaMaskEditorLog.h
--- a/Engine/Plugins/Experimental/Avalanche/Source/AvalancheMaskEditor/Private/AvaMaskEditorLog.h
+++ b/Engine/Plugins/Experimental/Avalanche/Source/AvalancheMaskEditor/Private/AvaMaskEditorLog.h
@@ -4,4 +4,4 @@
 
 #include "Logging/LogMacros.h"
 
-DECLARE_LOG_CATEGORY_EXTERN(LogAvaMaskEditor, VeryVerbose, All);
+DECLARE_LOG_CATEGORY_EXTERN(LogAvaMaskEditor, Log, All);
diff --git a/Engine/Plugins/Experimental/GeometryMask/Source/GeometryMaskEditor/Private/GeometryMaskEditorLog.h b/Engine/Plugins/Experimental/GeometryMask/Source/GeometryMaskEditor/Private/GeometryMaskEditorLog.h
--- a/Engine/Plugins/Experimental/GeometryMask/Source/GeometryMaskEditor/Private/GeometryMaskEditorLog.h
+++ b/Engine/Plugins/Experimental/GeometryMask/Source/GeometryMaskEditor/Private/GeometryMaskEditorLog.h
@@ -4,4 +4,4 @@
 
 #include "Logging/LogMacros.h"
 
-DECLARE_LOG_CATEGORY_EXTERN(LogGeometryMaskEditor, VeryVerbose, All);
+DECLARE_LOG_CATEGORY_EXTERN(LogGeometryMaskEditor, Log, All);
