[GeometryMask] Speculative fix for UObject lookup error in GeometryMaskCanvasResource

diff --git a/Engine/Plugins/Experimental/GeometryMask/Source/GeometryMask/Private/GeometryMaskCanvasResource.cpp b/Engine/Plugins/Experimental/GeometryMask/Source/GeometryMask/Private/GeometryMaskCanvasResource.cpp
--- a/Engine/Plugins/Experimental/GeometryMask/Source/GeometryMask/Private/GeometryMaskCanvasResource.cpp
+++ b/Engine/Plugins/Experimental/GeometryMask/Source/GeometryMask/Private/GeometryMaskCanvasResource.cpp
@@ -57,7 +57,10 @@ UGeometryMaskCanvasResource::UGeometryMaskCanvasResource()
 
 UGeometryMaskCanvasResource::~UGeometryMaskCanvasResource()
 {
-	if (!IsUnreachable() && IsValid(CanvasObject) && CanvasObject->Canvas)
+	if (this && IsValid(this)
+		&& CanvasObject
+		&& IsValid(CanvasObject)
+		&& CanvasObject->Canvas)
 	{
 		delete CanvasObject->Canvas;
 		CanvasObject->Canvas = nullptr;
