[GeometryMask] cis fix

#rnx

diff --git a/Engine/Plugins/Experimental/GeometryMask/Source/GeometryMask/Private/GeometryMaskCanvasResource.cpp b/Engine/Plugins/Experimental/GeometryMask/Source/GeometryMask/Private/GeometryMaskCanvasResource.cpp
--- a/Engine/Plugins/Experimental/GeometryMask/Source/GeometryMask/Private/GeometryMaskCanvasResource.cpp
+++ b/Engine/Plugins/Experimental/GeometryMask/Source/GeometryMask/Private/GeometryMaskCanvasResource.cpp
@@ -57,7 +57,7 @@ UGeometryMaskCanvasResource::UGeometryMaskCanvasResource()
 
 UGeometryMaskCanvasResource::~UGeometryMaskCanvasResource()
 {
-	if (this && IsValid(this)
+	if (IsValid(this)
 		&& CanvasObject
 		&& IsValid(CanvasObject)
 		&& CanvasObject->Canvas)
