[GeometryMask] CIS fix

#rnx

diff --git a/Engine/Plugins/Experimental/GeometryMask/Source/GeometryMask/Private/GeometryMaskCanvasResource.cpp b/Engine/Plugins/Experimental/GeometryMask/Source/GeometryMask/Private/GeometryMaskCanvasResource.cpp
--- a/Engine/Plugins/Experimental/GeometryMask/Source/GeometryMask/Private/GeometryMaskCanvasResource.cpp
+++ b/Engine/Plugins/Experimental/GeometryMask/Source/GeometryMask/Private/GeometryMaskCanvasResource.cpp
@@ -55,7 +55,7 @@ UGeometryMaskCanvasResource::UGeometryMaskCanvasResource()
 	PostProcess_DistanceField = MakeShared<FGeometryMaskPostProcess_DistanceField>(PostProcessParameters_DistanceField);
 }
 
-UGeometryMaskCanvasResource::~UGeometryMaskCanvasResource()
+void UGeometryMaskCanvasResource::BeginDestroy()
 {
 	if (IsValid(this)
 		&& CanvasObject
@@ -65,6 +65,8 @@ UGeometryMaskCanvasResource::~UGeometryMaskCanvasResource()
 		delete CanvasObject->Canvas;
 		CanvasObject->Canvas = nullptr;
 	}
+
+	Super::BeginDestroy();
 }
 
 const EGeometryMaskColorChannel UGeometryMaskCanvasResource::GetNextAvailableColorChannel() const
diff --git a/Engine/Plugins/Experimental/GeometryMask/Source/GeometryMask/Public/GeometryMaskCanvasResource.h b/Engine/Plugins/Experimental/GeometryMask/Source/GeometryMask/Public/GeometryMaskCanvasResource.h
--- a/Engine/Plugins/Experimental/GeometryMask/Source/GeometryMask/Public/GeometryMaskCanvasResource.h
+++ b/Engine/Plugins/Experimental/GeometryMask/Source/GeometryMask/Public/GeometryMaskCanvasResource.h
@@ -33,8 +33,9 @@ public:
 	static constexpr int32 MaxTextureSize = 8192;
 	
 	UGeometryMaskCanvasResource();
-	virtual ~UGeometryMaskCanvasResource() override;
-	
+
+	virtual void BeginDestroy() override;
+
 	/** Will return the first available color channel without a canvas assigned. EGeometryMaskColorChannel::None if not available. */
 	const EGeometryMaskColorChannel GetNextAvailableColorChannel() const;
 
