MotionDesign : GeometryMask
- Fixed mask modifier not updating when material is recompiled
- Fixed canvas and texture sample count do not match causing crash on mac

diff --git a/Engine/Plugins/VirtualProduction/Avalanche/Source/AvalancheMask/Private/Mask2D/AvaMask2DReadModifier.cpp b/Engine/Plugins/VirtualProduction/Avalanche/Source/AvalancheMask/Private/Mask2D/AvaMask2DReadModifier.cpp
--- a/Engine/Plugins/VirtualProduction/Avalanche/Source/AvalancheMask/Private/Mask2D/AvaMask2DReadModifier.cpp
+++ b/Engine/Plugins/VirtualProduction/Avalanche/Source/AvalancheMask/Private/Mask2D/AvaMask2DReadModifier.cpp
@@ -43,6 +43,14 @@ void UAvaMask2DReadModifier::PostEditChangeProperty(FPropertyChangedEvent& InPro
 
 	UAvaMask2DReadModifier::PropertyChangeDispatcher.OnPropertyChanged(this, InPropertyChangedEvent);
 }
+
+void UAvaMask2DReadModifier::OnMaterialCompiled(UMaterialInterface* InMaterial)
+{
+	if (InMaterial && InMaterial->GetTypedOuter<AActor>() == GetModifiedActor())
+	{
+		MarkModifierDirty();
+	}
+}
 #endif
 
 void UAvaMask2DReadModifier::Apply()
@@ -158,6 +166,24 @@ void UAvaMask2DReadModifier::OnModifierCDOSetup(FActorModifierCoreMetadata& InMe
 #endif
 }
 
+void UAvaMask2DReadModifier::OnModifierEnabled(EActorModifierCoreEnableReason InReason)
+{
+	Super::OnModifierEnabled(InReason);
+
+#if WITH_EDITOR
+	UMaterial::OnMaterialCompilationFinished().AddUObject(this, &UAvaMask2DReadModifier::OnMaterialCompiled);
+#endif
+}
+
+void UAvaMask2DReadModifier::OnModifierDisabled(EActorModifierCoreDisableReason InReason)
+{
+	Super::OnModifierDisabled(InReason);
+
+#if WITH_EDITOR
+	UMaterial::OnMaterialCompilationFinished().RemoveAll(this);
+#endif
+}
+
 void UAvaMask2DReadModifier::SetupMaskComponent(UActorComponent* InComponent)
 {
 	if (!InComponent)
diff --git a/Engine/Plugins/VirtualProduction/Avalanche/Source/AvalancheMask/Public/Mask2D/AvaMask2DReadModifier.h b/Engine/Plugins/VirtualProduction/Avalanche/Source/AvalancheMask/Public/Mask2D/AvaMask2DReadModifier.h
--- a/Engine/Plugins/VirtualProduction/Avalanche/Source/AvalancheMask/Public/Mask2D/AvaMask2DReadModifier.h
+++ b/Engine/Plugins/VirtualProduction/Avalanche/Source/AvalancheMask/Public/Mask2D/AvaMask2DReadModifier.h
@@ -46,12 +46,14 @@ public:
 #endif
 	//~ End UObject
 
-	const float GetBaseOpacity() const { return BaseOpacity; }
+	float GetBaseOpacity() const { return BaseOpacity; }
 	void SetBaseOpacity(float InBaseOpacity);
 
 protected:
 	//~ Begin UActorModifierCoreBase
 	virtual void OnModifierCDOSetup(FActorModifierCoreMetadata& InMetadata) override;
+	virtual void OnModifierEnabled(EActorModifierCoreEnableReason InReason) override;
+	virtual void OnModifierDisabled(EActorModifierCoreDisableReason InReason) override;
 	virtual void Apply() override;
 	//~ End UActorModifierCoreBase
 
@@ -73,5 +75,7 @@ private:
 #if WITH_EDITOR
 	/** Used for PECP */
 	static const TAvaPropertyChangeDispatcher<UAvaMask2DReadModifier> PropertyChangeDispatcher;
+
+	void OnMaterialCompiled(UMaterialInterface* InMaterial);
 #endif
 };
diff --git a/Engine/Plugins/VirtualProduction/GeometryMask/Source/GeometryMask/Private/GeometryMaskCanvasResource.cpp b/Engine/Plugins/VirtualProduction/GeometryMask/Source/GeometryMask/Private/GeometryMaskCanvasResource.cpp
--- a/Engine/Plugins/VirtualProduction/GeometryMask/Source/GeometryMask/Private/GeometryMaskCanvasResource.cpp
+++ b/Engine/Plugins/VirtualProduction/GeometryMask/Source/GeometryMask/Private/GeometryMaskCanvasResource.cpp
@@ -355,15 +355,7 @@ void UGeometryMaskCanvasResource::UpdateRenderParameters(
 
 	if (UCanvasRenderTarget2D* Texture = GetRenderTargetTexture())
 	{
-		if (bApplyBlur || bApplyDF)
-		{
-			// MSAA not supported for these effects (they need UAV access)
-			Texture->SetSampleCount(ETextureRenderTargetSampleCount::RTSC_1);
-		}
-		else
-		{
-			Texture->SetSampleCount(UE::GeometryMask::Private::RenderTargetSampleCount);
-		}
+		Texture->SetSampleCount(GetSampleCount());
 	}
 
 	// Effects may have changed viewport padding
@@ -432,7 +424,7 @@ UCanvasRenderTarget2D* UGeometryMaskCanvasResource::GetRenderTargetTexture()
 		RenderTargetTexture = NewObject<UCanvasRenderTarget2D>(this, ObjectName);
 		RenderTargetTexture->bForceLinearGamma = false;
 		RenderTargetTexture->bAutoGenerateMips = false;
-		RenderTargetTexture->SetSampleCount(UE::GeometryMask::Private::RenderTargetSampleCount);
+		RenderTargetTexture->SetSampleCount(GetSampleCount());
 		RenderTargetTexture->InitAutoFormat(MaxViewportSize.X, MaxViewportSize.Y);
 	}
 
@@ -447,6 +439,12 @@ void UGeometryMaskCanvasResource::Update(
 	Draw(InLevel, InView, InViewIndex);
 }
 
+ETextureRenderTargetSampleCount UGeometryMaskCanvasResource::GetSampleCount() const
+{
+	// MSAA not supported for these effects (they need UAV access)
+	return bApplyBlur || bApplyDF ? ETextureRenderTargetSampleCount::RTSC_1 : UE::GeometryMask::Private::RenderTargetSampleCount;
+}
+
 void UGeometryMaskCanvasResource::Draw(const ULevel* InLevel, FSceneView& InView, int32 InViewIndex)
 {
 	if (!InLevel || !InLevel->OwningWorld)
diff --git a/Engine/Plugins/VirtualProduction/GeometryMask/Source/GeometryMask/Public/GeometryMaskCanvasResource.h b/Engine/Plugins/VirtualProduction/GeometryMask/Source/GeometryMask/Public/GeometryMaskCanvasResource.h
--- a/Engine/Plugins/VirtualProduction/GeometryMask/Source/GeometryMask/Public/GeometryMaskCanvasResource.h
+++ b/Engine/Plugins/VirtualProduction/GeometryMask/Source/GeometryMask/Public/GeometryMaskCanvasResource.h
@@ -16,6 +16,7 @@ class IGeometryMaskPostProcess;
 class UCanvas;
 class UCanvasRenderTarget2D;
 enum class EGeometryMaskColorChannel : uint8;
+enum class ETextureRenderTargetSampleCount : uint8;
 
 using FOnGeometryMaskCanvasDraw = TMulticastDelegate<void(const FGeometryMaskDrawingContext&, FCanvas*)>;
 
@@ -75,6 +76,8 @@ public:
 	void Update(const ULevel* InLevel, FSceneView& InView, int32 InViewIndex = 0);
 
 private:
+	ETextureRenderTargetSampleCount GetSampleCount() const;
+
 	/** Draws all writers to the canvas. */
 	void Draw(const ULevel* InLevel, FSceneView& InView, int32 InViewIndex = 0);
 
@@ -84,7 +87,6 @@ private:
 
 	[[maybe_unused]] bool RemoveInvalidDrawingContexts();
 
-private:
 	UPROPERTY(NoClear, EditFixedSize, meta = (EditFixedOrder))
 	TMap<EGeometryMaskColorChannel, FGeometryMaskCanvasId> DependentCanvasIds;
 
