Motion Design: Fixed issues when multiple masking elements are drawn. The issue stemmed from the draw code pushing index values that were 'relative' to the current shape being drawn, rather than the actual indices in the BatchedElement vertex buffer.
Additionally, every canvas was being updated per world regardless of the world which was causing issues when multiple worlds were present (e.g. the In-Editor World, and Rundown).

diff --git a/Engine/Plugins/VirtualProduction/GeometryMask/Source/GeometryMask/Private/GeometryMaskSubsystem.cpp b/Engine/Plugins/VirtualProduction/GeometryMask/Source/GeometryMask/Private/GeometryMaskSubsystem.cpp
--- a/Engine/Plugins/VirtualProduction/GeometryMask/Source/GeometryMask/Private/GeometryMaskSubsystem.cpp
+++ b/Engine/Plugins/VirtualProduction/GeometryMask/Source/GeometryMask/Private/GeometryMaskSubsystem.cpp
@@ -83,16 +83,23 @@ void UGeometryMaskSubsystem::UpdateLevel(const ULevel* InLevel, UGeometryMaskWor
 	{
 		FSceneView* MutableSceneView = const_cast<FSceneView*>(View);
 
+		TArray<UGeometryMaskCanvasResource*, TInlineAllocator<1>> UsedResources;
+
 		if (const FGeometryMaskLevelState* LevelState = InWorldSubsystem->FindLevelState(InLevel))
 		{
 			for (const TPair<FName, TObjectPtr<UGeometryMaskCanvas>>& NamedCanvas : LevelState->NamedCanvases)
 			{
-				NamedCanvas.Value->Update(InLevel, *MutableSceneView);		
+				NamedCanvas.Value->Update(InLevel, *MutableSceneView);
+
+				if (UGeometryMaskCanvasResource* Resource = NamedCanvas.Value->GetResourceMutable())
+				{
+					UsedResources.AddUnique(Resource);
+				}
 			}
 		}
 
 		// Updates the texture resource
-		for (const TObjectPtr<UGeometryMaskCanvasResource>& Resource : CanvasResources)
+		for (UGeometryMaskCanvasResource* Resource : UsedResources)
 		{
 			Resource->Update(InLevel, *MutableSceneView, ViewIndex);
 		}
diff --git a/Engine/Plugins/VirtualProduction/GeometryMask/Source/GeometryMask/Private/GeometryMaskWriteComponent.cpp b/Engine/Plugins/VirtualProduction/GeometryMask/Source/GeometryMask/Private/GeometryMaskWriteComponent.cpp
--- a/Engine/Plugins/VirtualProduction/GeometryMask/Source/GeometryMask/Private/GeometryMaskWriteComponent.cpp
+++ b/Engine/Plugins/VirtualProduction/GeometryMask/Source/GeometryMask/Private/GeometryMaskWriteComponent.cpp
@@ -99,21 +99,13 @@ void UGeometryMaskWriteMeshComponent::DrawToCanvas(FCanvas* InCanvas)
 			const FTransform LocalToWorld = ComponentTransforms[KeyName];
 			const FGeometryMaskBatchElementData& MeshBatchElementData = CachedMeshData[KeyName];
 
+			InCanvas->PushAbsoluteTransform(LocalToWorld.ToMatrixWithScale());
 			{
-				FScopeLock CanvasLock(&CanvasCS);
-				InCanvas->PushAbsoluteTransform(LocalToWorld.ToMatrixWithScale());	
-			}
-			
-			{
-				FBatchedElements* CanvasTriangleBatchedElements = nullptr;				
-				{
-					FScopeLock CanvasLock(&CanvasCS);
-					CanvasTriangleBatchedElements = InCanvas->GetBatchedElements(FCanvas::ET_Triangle);
-				}
-				
-				CanvasTriangleBatchedElements->ReserveVertices(MeshBatchElementData.Vertices.Num());
-				CanvasTriangleBatchedElements->ReserveTriangles(MeshBatchElementData.NumTriangles, GWhiteTexture, ElementBlendMode);
-				
+				FBatchedElements* CanvasTriangleBatchedElements = InCanvas->GetBatchedElements(FCanvas::ET_Triangle);
+
+				CanvasTriangleBatchedElements->AddReserveVertices(MeshBatchElementData.Vertices.Num());
+				CanvasTriangleBatchedElements->AddReserveTriangles(MeshBatchElementData.NumTriangles, GWhiteTexture, ElementBlendMode);
+
 				auto AddVertex = [CanvasTriangleBatchedElements, &MeshBatchElementData, CanvasHitProxyId, &Color](int32 VertexIdx) -> int32
 				{
 					return CanvasTriangleBatchedElements->AddVertexf(
@@ -123,9 +115,16 @@ void UGeometryMaskWriteMeshComponent::DrawToCanvas(FCanvas* InCanvas)
 						CanvasHitProxyId);
 				};
 
+				// Store the index of the first vertex to be able to convert the MeshBatchElementData vertex indices to
+				// map to the actual vertices added to the BatchedElements.
+				int32 InitialIndex = 0;
+
+				if (MeshBatchElementData.Vertices.Num() > 0)
 				{
 					TRACE_CPUPROFILER_EVENT_SCOPE(UGeometryMaskWriteMeshComponent::DrawToCanvas::WriteToCanvas::AddVertices);
-					for (int32 VertexIdx = 0; VertexIdx < MeshBatchElementData.Vertices.Num(); ++VertexIdx)
+
+					InitialIndex = AddVertex(0);
+					for (int32 VertexIdx = 1; VertexIdx < MeshBatchElementData.Vertices.Num(); ++VertexIdx)
 					{
 						AddVertex(VertexIdx);
 					}
@@ -135,19 +134,18 @@ void UGeometryMaskWriteMeshComponent::DrawToCanvas(FCanvas* InCanvas)
 					TRACE_CPUPROFILER_EVENT_SCOPE(UGeometryMaskWriteMeshComponent::DrawToCanvas::WriteToCanvas::AddTriangles);
 					for (int32 VertexIdx = 0; VertexIdx <= MeshBatchElementData.Indices.Num() - 3; VertexIdx += 3)
 					{
-						const int32 V0 = MeshBatchElementData.Indices[VertexIdx];
-						const int32 V1 = MeshBatchElementData.Indices[VertexIdx + 1];
-						const int32 V2 = MeshBatchElementData.Indices[VertexIdx + 2];
+						// MeshBatchElementData.Indices are 'relative' in the sense that do not consider possible existing 
+						// vertices within BatchedElements prior to its vertices being added.
+						// So these 'local' indices are converted to be the proper indices of the vertices that were added above. 
+						const int32 V0 = InitialIndex + MeshBatchElementData.Indices[VertexIdx];
+						const int32 V1 = InitialIndex + MeshBatchElementData.Indices[VertexIdx + 1];
+						const int32 V2 = InitialIndex + MeshBatchElementData.Indices[VertexIdx + 2];
 
 						CanvasTriangleBatchedElements->AddTriangle(V0, V1, V2, GWhiteTexture, ElementBlendMode);
 					}
 				}
 			}
-			
-			{
-				FScopeLock CanvasLock(&CanvasCS);
-				InCanvas->PopTransform();
-			}
+			InCanvas->PopTransform();
 		}
 
 		CachedComponentsWeak = ComponentsToKeep;
diff --git a/Engine/Plugins/VirtualProduction/GeometryMask/Source/GeometryMask/Public/GeometryMaskCanvas.h b/Engine/Plugins/VirtualProduction/GeometryMask/Source/GeometryMask/Public/GeometryMaskCanvas.h
--- a/Engine/Plugins/VirtualProduction/GeometryMask/Source/GeometryMask/Public/GeometryMaskCanvas.h
+++ b/Engine/Plugins/VirtualProduction/GeometryMask/Source/GeometryMask/Public/GeometryMaskCanvas.h
@@ -112,7 +112,16 @@ public:
 	/** Updates the canvas, intended to be called every frame. */
 	void Update(const ULevel* InLevel, FSceneView& InView);
 	
-	const UGeometryMaskCanvasResource* GetResource() const { return CanvasResource; }
+	const UGeometryMaskCanvasResource* GetResource() const
+	{
+		return CanvasResource;
+	}
+
+	UGeometryMaskCanvasResource* GetResourceMutable()
+	{
+		return CanvasResource;
+	}
+
 	void AssignResource(UGeometryMaskCanvasResource* InResource, EGeometryMaskColorChannel InColorChannel);
 	void FreeResource();
 
diff --git a/Engine/Plugins/VirtualProduction/GeometryMask/Source/GeometryMask/Public/GeometryMaskWriteComponent.h b/Engine/Plugins/VirtualProduction/GeometryMask/Source/GeometryMask/Public/GeometryMaskWriteComponent.h
--- a/Engine/Plugins/VirtualProduction/GeometryMask/Source/GeometryMask/Public/GeometryMaskWriteComponent.h
+++ b/Engine/Plugins/VirtualProduction/GeometryMask/Source/GeometryMask/Public/GeometryMaskWriteComponent.h
@@ -100,7 +100,4 @@ protected:
 	
 	UPROPERTY(DuplicateTransient)
 	TMap<FName, FGeometryMaskBatchElementData> CachedMeshData;
-
-	/** To protect async-unsafe calls. */
-	FCriticalSection CanvasCS;
 };
