MotionDesign : GeometryMask
- Add ability to write the same mesh multiple times if it's used in multiple component under the actor

#DisabledValidations NonUnityAndStaticAnalysis

diff --git a/Engine/Plugins/VirtualProduction/GeometryMask/Source/GeometryMask/Private/GeometryMaskWriteComponent.cpp b/Engine/Plugins/VirtualProduction/GeometryMask/Source/GeometryMask/Private/GeometryMaskWriteComponent.cpp
--- a/Engine/Plugins/VirtualProduction/GeometryMask/Source/GeometryMask/Private/GeometryMaskWriteComponent.cpp
+++ b/Engine/Plugins/VirtualProduction/GeometryMask/Source/GeometryMask/Private/GeometryMaskWriteComponent.cpp
@@ -58,46 +58,40 @@ void UGeometryMaskWriteMeshComponent::DrawToCanvas(FCanvas* InCanvas)
 		ElementBlendMode = SE_BLEND_Modulate;
 	}
 
-	// @todo: cached mesh data keyed by SM name, but transforms also tied to SM, not SMC - will break if same SM used in multiple SMC's
-
 	// Write to Canvas BatchElements
 	{
 		FHitProxyId CanvasHitProxyId = InCanvas->GetHitProxyId();
 
-		TArray<FName> KeyNames;
-		CachedMeshData.GenerateKeyArray(KeyNames);
-
-		TMap<FName, TWeakObjectPtr<USceneComponent>> ComponentsToKeep;
-		ComponentsToKeep.Reserve(KeyNames.Num());
-
-		TMap<FName, FTransform> ComponentTransforms;
-		ComponentTransforms.Reserve(KeyNames.Num());
+		TMap<TWeakObjectPtr<UPrimitiveComponent>, FName> ComponentsToKeep;
+		ComponentsToKeep.Reserve(CachedComponents.Num());
 
 		TRACE_CPUPROFILER_EVENT_SCOPE(UGeometryMaskWriteMeshComponent::DrawToCanvas::WriteToCanvas);
 
 		// Flag valid components, can't do async due to UObject access
-		for (FName KeyName : KeyNames)
+		for (const TPair<TWeakObjectPtr<UPrimitiveComponent>, FName>& CachedComponentWeak : CachedComponents)
 		{
-			if (CachedComponentsWeak.Contains(KeyName))
+			if (UPrimitiveComponent* CachedComponent = CachedComponentWeak.Key.Get())
 			{
-				TWeakObjectPtr<USceneComponent>& CachedComponentWeak = CachedComponentsWeak[KeyName];
-				if (const USceneComponent* CachedComponent = CachedComponentWeak.Get())
+				if (CachedMeshData.Contains(CachedComponentWeak.Value))
 				{
-					ComponentsToKeep.Emplace(KeyName, CachedComponentWeak);
-					ComponentTransforms.Emplace(KeyName, CachedComponent->GetComponentToWorld());
+					ComponentsToKeep.Emplace(CachedComponent, CachedComponentWeak.Value);
 				}
 			}
 		}
 
-		// Regenerate keys for valid components only
-		ComponentsToKeep.GenerateKeyArray(KeyNames);
-
-		for (FName KeyName : KeyNames)
+		for (const TPair<TWeakObjectPtr<UPrimitiveComponent>, FName>& ComponentPair : ComponentsToKeep)
 		{
 			TRACE_CPUPROFILER_EVENT_SCOPE(UGeometryMaskWriteMeshComponent::DrawToCanvas::WriteToCanvas::Task);
 
-			const FTransform LocalToWorld = ComponentTransforms[KeyName];
-			const FGeometryMaskBatchElementData& MeshBatchElementData = CachedMeshData[KeyName];
+			const UPrimitiveComponent* PrimitiveComponent = ComponentPair.Key.Get();
+
+			if (!PrimitiveComponent)
+			{
+				continue;
+			}
+
+			const FTransform LocalToWorld = PrimitiveComponent->GetComponentToWorld();
+			const FGeometryMaskBatchElementData& MeshBatchElementData = CachedMeshData[ComponentPair.Value];
 
 			InCanvas->PushAbsoluteTransform(LocalToWorld.ToMatrixWithScale());
 			{
@@ -148,7 +142,7 @@ void UGeometryMaskWriteMeshComponent::DrawToCanvas(FCanvas* InCanvas)
 			InCanvas->PopTransform();
 		}
 
-		CachedComponentsWeak = ComponentsToKeep;
+		CachedComponents = ComponentsToKeep;
 	}
 }
 
@@ -173,7 +167,7 @@ void UGeometryMaskWriteMeshComponent::UpdateCachedData()
 		}
 
 		LastPrimitiveComponentCount = PrimitiveComponents.Num();
-		CachedComponentsWeak.Reserve(LastPrimitiveComponentCount);
+		CachedComponents.Reserve(LastPrimitiveComponentCount);
 		CachedMeshData.Reserve(LastPrimitiveComponentCount);
 
 		UpdateCachedStaticMeshData(PrimitiveComponents);
@@ -206,35 +200,30 @@ void UGeometryMaskWriteMeshComponent::UpdateCachedStaticMeshData(TConstArrayView
 	// Remove built/already cached
 	StaticMeshComponents.SetNum(Algo::RemoveIf(StaticMeshComponents, [this](const UStaticMeshComponent* InStaticMeshComponent)
 	{
-		if (const UStaticMesh* StaticMesh = InStaticMeshComponent->GetStaticMesh())
+		const UStaticMesh* StaticMesh = InStaticMeshComponent->GetStaticMesh();
+	
+		if (!StaticMesh || !StaticMesh->HasValidRenderData())
 		{
-			const FName ComponentKey = StaticMesh->GetFName();
-			if (const FGeometryMaskBatchElementData* CachedData = CachedMeshData.Find(ComponentKey))
-			{
-				if (!CachedComponentsWeak.Contains(ComponentKey))
-				{
-					return false; // Cached component invalid, don't remove from build list to ensure it's re-cached
-				}
-				
-				if (CachedData->Vertices.Num() == StaticMesh->GetNumVertices(0))
-				{
-					return true; // Cached, remove from "to build" list							
-				}
-			}
+			return true;
+		}
+
+		const FName ResourceName = StaticMesh->GetFName();
+		if (CachedComponents.Contains(InStaticMeshComponent) && CachedMeshData.Contains(ResourceName))
+		{
+			return true;
 		}
 
 		return false;
 	}));
 
 	// Subscribe to change events
+#if WITH_EDITOR
+	for (UStaticMeshComponent* StaticMeshComponent : StaticMeshComponents)
 	{
-		#if WITH_EDITOR
-		for (UStaticMeshComponent* StaticMeshComponent : StaticMeshComponents)
-		{
-			StaticMeshComponent->OnStaticMeshChanged().AddUObject(this, &UGeometryMaskWriteMeshComponent::OnStaticMeshChanged);
-		}
-		#endif
+		StaticMeshComponent->OnStaticMeshChanged().RemoveAll(this);
+		StaticMeshComponent->OnStaticMeshChanged().AddUObject(this, &UGeometryMaskWriteMeshComponent::OnStaticMeshChanged);
 	}
+#endif
 
 	// Used for cache lookup
 	TArray<FName> StaticMeshObjectNames;
@@ -267,9 +256,10 @@ void UGeometryMaskWriteMeshComponent::UpdateCachedStaticMeshData(TConstArrayView
 						continue;
 					}
 
-					StaticMeshObjectNames.Add(StaticMesh->GetFName());
+					const FName ResourceName = StaticMesh->GetFName();
+					StaticMeshObjectNames.Add(ResourceName);
 					StaticMeshResources.Add(&RenderData->LODResources[0]);
-					CachedComponentsWeak.Add(StaticMesh->GetFName(), StaticMeshComponent);
+					CachedComponents.Add(StaticMeshComponent, ResourceName);
 				}
 			}
 		}
@@ -327,34 +317,32 @@ void UGeometryMaskWriteMeshComponent::UpdateCachedDynamicMeshData(TConstArrayVie
 	// Remove built/already cached
 	DynamicMeshComponents.SetNum(Algo::RemoveIf(DynamicMeshComponents, [this](UDynamicMeshComponent* InDynamicMeshComponent)
 	{
-		if (const UDynamicMesh* DynamicMesh = InDynamicMeshComponent->GetDynamicMesh())
+		const UDynamicMesh* DynamicMesh = InDynamicMeshComponent->GetDynamicMesh();
+	
+		if (!DynamicMesh || !DynamicMesh->GetMeshPtr())
 		{
-			const FName ComponentKey = InDynamicMeshComponent->GetFName();
-			if (const FGeometryMaskBatchElementData* CachedData = CachedMeshData.Find(ComponentKey))
-			{
-				if (!CachedComponentsWeak.Contains(ComponentKey))
-				{
-					return false; // Cached component invalid, don't remove from build list to ensure it's re-cached
-				}
-				
-				if (CachedData->Vertices.Num() == DynamicMesh->GetMeshRef().VertexCount()
-					&& CachedData->ChangeStamp == DynamicMesh->GetMeshRef().GetChangeStamp())
-				{
-					return true; // Cached, remove from "to build" list							
-				}
-			}
+			return true;
+		}
+
+		const FName ResourceName = InDynamicMeshComponent->GetFName();
+		if (CachedComponents.Contains(InDynamicMeshComponent) && CachedMeshData.Contains(ResourceName))
+		{
+			return true;
 		}
 
 		return false;
 	}));
 
 	// Subscribe to change events
+	for (UDynamicMeshComponent* DynamicMeshComponent : DynamicMeshComponents)
 	{
-		for (UDynamicMeshComponent* DynamicMeshComponent : DynamicMeshComponents)
+		if (FDynamicMesh3* DynamicMesh3 = DynamicMeshComponent->GetMesh())
 		{
-			DynamicMeshComponent->GetMesh()->SetShapeChangeStampEnabled(true);			
-			DynamicMeshComponent->OnMeshChanged.AddUObject(this, &UGeometryMaskWriteMeshComponent::OnDynamicMeshChanged, DynamicMeshComponent);
+			DynamicMesh3->SetShapeChangeStampEnabled(true);
 		}
+
+		DynamicMeshComponent->OnMeshChanged.RemoveAll(this);
+		DynamicMeshComponent->OnMeshChanged.AddUObject(this, &UGeometryMaskWriteMeshComponent::OnDynamicMeshChanged, DynamicMeshComponent);
 	}
 
 	// Used for cache lookup
@@ -373,9 +361,15 @@ void UGeometryMaskWriteMeshComponent::UpdateCachedDynamicMeshData(TConstArrayVie
 			{
 				if (FDynamicMesh3* DynamicMesh = DynamicMeshObject->GetMeshPtr())
 				{
-					DynamicMeshObjectNames.Add(DynamicMeshComponent->GetFName());
+					if (DynamicMesh->TriangleCount() == 0)
+					{
+						continue;
+					}
+
+					const FName ResourceName = DynamicMeshComponent->GetFName();
+					DynamicMeshObjectNames.Add(ResourceName);
 					DynamicMeshes.Add(DynamicMesh);
-					CachedComponentsWeak.Add(DynamicMeshComponent->GetFName(), DynamicMeshComponent);
+					CachedComponents.Add(DynamicMeshComponent, ResourceName);
 				}
 			}
 		}
@@ -475,6 +469,6 @@ bool UGeometryMaskWriteMeshComponent::Cleanup()
 void UGeometryMaskWriteMeshComponent::ResetCachedData()
 {
 	LastPrimitiveComponentCount = -1;
-	CachedComponentsWeak.Reset();
+	CachedComponents.Reset();
 	CachedMeshData.Reset();
 }
diff --git a/Engine/Plugins/VirtualProduction/GeometryMask/Source/GeometryMask/Public/GeometryMaskWriteComponent.h b/Engine/Plugins/VirtualProduction/GeometryMask/Source/GeometryMask/Public/GeometryMaskWriteComponent.h
--- a/Engine/Plugins/VirtualProduction/GeometryMask/Source/GeometryMask/Public/GeometryMaskWriteComponent.h
+++ b/Engine/Plugins/VirtualProduction/GeometryMask/Source/GeometryMask/Public/GeometryMaskWriteComponent.h
@@ -10,6 +10,7 @@
 
 class UCanvasRenderTarget2D;
 class UDynamicMeshComponent;
+class UPrimitiveComponent;
 class UStaticMeshComponent;
 
 /** Mesh data for a single mesh, in local space. */
@@ -95,9 +96,15 @@ protected:
 	UPROPERTY(Transient, DuplicateTransient)
 	int32 LastPrimitiveComponentCount;
 
+	UE_DEPRECATED(5.7, "Replaced by CachedComponents")
+	UPROPERTY(Transient, meta=(DeprecatedProperty, DeprecatedMessage="Replaced by CachedComponents"))
+	TMap<FName, TWeakObjectPtr<UPrimitiveComponent>> CachedComponentsWeak;
+
+	/** Cached primitive component to mesh name, multiple components can use same mesh data */
 	UPROPERTY(DuplicateTransient)
-	TMap<FName, TWeakObjectPtr<USceneComponent>> CachedComponentsWeak;
-	
+	TMap<TWeakObjectPtr<UPrimitiveComponent>, FName> CachedComponents;
+
+	/** Cached mesh name to data (one entry per unique mesh) */
 	UPROPERTY(DuplicateTransient)
 	TMap<FName, FGeometryMaskBatchElementData> CachedMeshData;
 };
