Geometry Mask: fix issue where active PIE prevents geometry mask view extension to work in rundown worlds

diff --git a/Engine/Plugins/VirtualProduction/GeometryMask/Source/GeometryMask/Private/GeometryMaskSVE.cpp b/Engine/Plugins/VirtualProduction/GeometryMask/Source/GeometryMask/Private/GeometryMaskSVE.cpp
--- a/Engine/Plugins/VirtualProduction/GeometryMask/Source/GeometryMask/Private/GeometryMaskSVE.cpp
+++ b/Engine/Plugins/VirtualProduction/GeometryMask/Source/GeometryMask/Private/GeometryMaskSVE.cpp
@@ -38,21 +38,21 @@ bool FGeometryMaskSceneViewExtension::IsActiveThisFrame_Internal(const FSceneVie
 #if WITH_EDITOR
 	if (GEditor)
 	{
-		if (GEditor->IsSimulatingInEditor())
+		if (const UWorld* const ContextWorld = Context.GetWorld())
 		{
-			if (Context.GetWorld()->WorldType == EWorldType::Editor)
+			if (GEditor->IsSimulatingInEditor() && ContextWorld->WorldType == EWorldType::Editor)
 			{
 				return true;
 			}
-		}
-		
-		if (GEditor->PlayWorld)
-		{
-			bIsActive = Context.GetWorld()->WorldType == EWorldType::PIE;
-		}
-		else
-		{
-			bIsActive = Context.GetWorld()->WorldType == EWorldType::Editor || GetWorld()->WorldType == EWorldType::Game;
+
+			if (GEditor->PlayWorld)
+			{
+				bIsActive = ContextWorld->WorldType == EWorldType::PIE || ContextWorld->WorldType == EWorldType::Game;
+			}
+			else
+			{
+				bIsActive = ContextWorld->WorldType == EWorldType::Editor || ContextWorld->WorldType == EWorldType::Game;
+			}
 		}
 	}
 #endif
