Motion Design: improve performance of GeometryMask by preallocating the canvas object in ctor instead.

#Synced-CL 48030442

diff --git a/Engine/Plugins/VirtualProduction/GeometryMask/Source/GeometryMask/Private/GeometryMaskCanvasResource.cpp b/Engine/Plugins/VirtualProduction/GeometryMask/Source/GeometryMask/Private/GeometryMaskCanvasResource.cpp
--- a/Engine/Plugins/VirtualProduction/GeometryMask/Source/GeometryMask/Private/GeometryMaskCanvasResource.cpp
+++ b/Engine/Plugins/VirtualProduction/GeometryMask/Source/GeometryMask/Private/GeometryMaskCanvasResource.cpp
@@ -38,6 +38,9 @@ namespace UE::GeometryMask::Private
 UGeometryMaskCanvasResource::UGeometryMaskCanvasResource()
 	: DefaultDrawingContext()
 {
+	// Note: Here transient does not mean that the UObject is transient, it's just that it's assigned to a transient property.
+	CanvasObject = CreateDefaultSubobject<UCanvas>(TEXT("MaskCanvas"), /*bIsTransient*/true);
+
 	// @note: We omit the Alpha channel as it's not uniformly supported
 	DependentCanvasIds = { { EGeometryMaskColorChannel::Red, FGeometryMaskCanvasId(EForceInit::ForceInit) }, { EGeometryMaskColorChannel::Green, FGeometryMaskCanvasId(EForceInit::ForceInit) }, { EGeometryMaskColorChannel::Blue, FGeometryMaskCanvasId(EForceInit::ForceInit) } };
 
@@ -55,10 +58,7 @@ UGeometryMaskCanvasResource::UGeometryMaskCanvasResource()
 
 void UGeometryMaskCanvasResource::BeginDestroy()
 {
-	if (IsValidChecked(this)
-		&& CanvasObject
-		&& IsValid(CanvasObject)
-		&& CanvasObject->Canvas)
+	if (CanvasObject && CanvasObject->Canvas)
 	{
 		delete CanvasObject->Canvas;
 		CanvasObject->Canvas = nullptr;
@@ -443,7 +443,7 @@ ETextureRenderTargetSampleCount UGeometryMaskCanvasResource::GetSampleCount() co
 
 void UGeometryMaskCanvasResource::Draw(const ULevel* InLevel, FSceneView& InView, int32 InViewIndex)
 {
-	if (!InLevel || !InLevel->OwningWorld)
+	if (!InLevel || !InLevel->OwningWorld || !CanvasObject)
 	{
 		return;
 	}
@@ -453,10 +453,8 @@ void UGeometryMaskCanvasResource::Draw(const ULevel* InLevel, FSceneView& InView
 	if (UCanvasRenderTarget2D* Texture = GetRenderTargetTexture())
 	{
 		// First time initialization
-		if (!CanvasObject)
+		if (!CanvasObject->Canvas)
 		{
-			CanvasObject = NewObject<UCanvas>(GetTransientPackage());
-
 			FTextureRenderTargetResource* RenderTargetResource = Texture->GameThread_GetRenderTargetResource();
 			FCanvas* NewCanvas = new FCanvas(
 				RenderTargetResource,
@@ -515,23 +513,20 @@ void UGeometryMaskCanvasResource::Draw(const ULevel* InLevel, FSceneView& InView
 		
 		// End
 		{
-			if (CanvasObject && CanvasObject->Canvas)
+			CanvasObject->Canvas->Flush_GameThread();
+
+			// Post Process
+			if (bApplyDF || bApplyBlur)
 			{
-				CanvasObject->Canvas->Flush_GameThread();
-				
-				// Post Process
-				if (bApplyDF || bApplyBlur)
+				FRenderTarget* RenderTarget = CanvasObject->Canvas->GetRenderTarget();
+				if (bApplyDF)
 				{
-					FRenderTarget* RenderTarget = CanvasObject->Canvas->GetRenderTarget();
-					if (bApplyDF)
-					{
-						PostProcess_DistanceField->Execute(RenderTarget);	
-					}
+					PostProcess_DistanceField->Execute(RenderTarget);	
+				}
 
-					if (bApplyBlur)
-					{
-						PostProcess_Blur->Execute(RenderTarget);
-					}
+				if (bApplyBlur)
+				{
+					PostProcess_Blur->Execute(RenderTarget);
 				}
 			}
 		}
