Motion Design: Geometry Mask canvas resource now uses FCanvas directly as UCanvas has expensive calls like GetDisplayMetrics that the resource doesn't need.

#Synced-CL 48355279

diff --git a/Engine/Plugins/VirtualProduction/GeometryMask/Source/GeometryMask/Private/GeometryMaskCanvasResource.cpp b/Engine/Plugins/VirtualProduction/GeometryMask/Source/GeometryMask/Private/GeometryMaskCanvasResource.cpp
--- a/Engine/Plugins/VirtualProduction/GeometryMask/Source/GeometryMask/Private/GeometryMaskCanvasResource.cpp
+++ b/Engine/Plugins/VirtualProduction/GeometryMask/Source/GeometryMask/Private/GeometryMaskCanvasResource.cpp
@@ -38,9 +38,6 @@ namespace UE::GeometryMask::Private
 UGeometryMaskCanvasResource::UGeometryMaskCanvasResource()
 	: DefaultDrawingContext()
 {
-	// Note: Here transient does not mean that the UObject is transient, it's just that it's assigned to a transient property.
-	CanvasObject = CreateDefaultSubobject<UCanvas>(TEXT("MaskCanvas"), /*bIsTransient*/true);
-
 	// @note: We omit the Alpha channel as it's not uniformly supported
 	DependentCanvasIds = { { EGeometryMaskColorChannel::Red, FGeometryMaskCanvasId(EForceInit::ForceInit) }, { EGeometryMaskColorChannel::Green, FGeometryMaskCanvasId(EForceInit::ForceInit) }, { EGeometryMaskColorChannel::Blue, FGeometryMaskCanvasId(EForceInit::ForceInit) } };
 
@@ -58,12 +55,7 @@ UGeometryMaskCanvasResource::UGeometryMaskCanvasResource()
 
 void UGeometryMaskCanvasResource::BeginDestroy()
 {
-	if (CanvasObject && CanvasObject->Canvas)
-	{
-		delete CanvasObject->Canvas;
-		CanvasObject->Canvas = nullptr;
-	}
-
+	Canvas.Reset();
 	Super::BeginDestroy();
 }
 
@@ -443,7 +435,7 @@ ETextureRenderTargetSampleCount UGeometryMaskCanvasResource::GetSampleCount() co
 
 void UGeometryMaskCanvasResource::Draw(const ULevel* InLevel, FSceneView& InView, int32 InViewIndex)
 {
-	if (!InLevel || !InLevel->OwningWorld || !CanvasObject)
+	if (!InLevel || !InLevel->OwningWorld)
 	{
 		return;
 	}
@@ -453,18 +445,15 @@ void UGeometryMaskCanvasResource::Draw(const ULevel* InLevel, FSceneView& InView
 	if (UCanvasRenderTarget2D* Texture = GetRenderTargetTexture())
 	{
 		// First time initialization
-		if (!CanvasObject->Canvas)
+		if (!Canvas.IsValid())
 		{
-			FTextureRenderTargetResource* RenderTargetResource = Texture->GameThread_GetRenderTargetResource();
-			FCanvas* NewCanvas = new FCanvas(
-				RenderTargetResource,
+			Canvas = MakeShared<FCanvas>(
+				Texture->GameThread_GetRenderTargetResource(),
 				nullptr,
 				InLevel->OwningWorld,
 				InLevel->OwningWorld->GetFeatureLevel(),
 				// Draw immediately so that interleaved SetVectorParameter (etc) function calls work as expected
 				FCanvas::CDM_ImmediateDrawing);
-
-			CanvasObject->Init(Texture->SizeX, Texture->SizeY, &InView, NewCanvas);
 		}
 
 		FGeometryMaskDrawingContext* DrawingContextPtr = GetDrawingContextForLevel(InLevel, InViewIndex);
@@ -485,40 +474,36 @@ void UGeometryMaskCanvasResource::Draw(const ULevel* InLevel, FSceneView& InView
 
 			InLevel->OwningWorld->FlushDeferredParameterCollectionInstanceUpdates();
 
-			CanvasObject->SizeX = Texture->SizeX;
-			CanvasObject->SizeY = Texture->SizeY;
-			CanvasObject->Update();
-			CanvasObject->SetView(&InView);
-			CanvasObject->Canvas->Clear(Texture->ClearColor);
-
+			Canvas->SetParentCanvasSize(FIntPoint(Texture->SizeX, Texture->SizeY));
+			Canvas->Clear(Texture->ClearColor);
 			Texture->UpdateResourceImmediate(/** Clear */true);
 		}
 
 		// Contents
 		{
 			// Store current transform
-			const FMatrix CanvasMatrix = CanvasObject->Canvas->GetBottomTransform();
+			const FMatrix CanvasMatrix = Canvas->GetBottomTransform();
 
 			// Set to World->Viewport transform
-			CanvasObject->Canvas->SetBaseTransform(DrawingContextRef.ViewProjectionMatrix);
+			Canvas->SetBaseTransform(DrawingContextRef.ViewProjectionMatrix);
 			{
 				TRACE_CPUPROFILER_EVENT_SCOPE(UGeometryMaskCanvasResource::Draw::OnDrawToCanvas);
 				
-				OnDrawToCanvas().Broadcast(DrawingContextRef, CanvasObject->Canvas);
+				OnDrawToCanvas().Broadcast(DrawingContextRef, Canvas.Get());
 			}
 
 			// Restore original transform
-			CanvasObject->Canvas->SetBaseTransform(CanvasMatrix);
+			Canvas->SetBaseTransform(CanvasMatrix);
 		}
 		
 		// End
 		{
-			CanvasObject->Canvas->Flush_GameThread();
+			Canvas->Flush_GameThread();
 
 			// Post Process
 			if (bApplyDF || bApplyBlur)
 			{
-				FRenderTarget* RenderTarget = CanvasObject->Canvas->GetRenderTarget();
+				FRenderTarget* RenderTarget = Canvas->GetRenderTarget();
 				if (bApplyDF)
 				{
 					PostProcess_DistanceField->Execute(RenderTarget);	
diff --git a/Engine/Plugins/VirtualProduction/GeometryMask/Source/GeometryMask/Public/GeometryMaskCanvasResource.h b/Engine/Plugins/VirtualProduction/GeometryMask/Source/GeometryMask/Public/GeometryMaskCanvasResource.h
--- a/Engine/Plugins/VirtualProduction/GeometryMask/Source/GeometryMask/Public/GeometryMaskCanvasResource.h
+++ b/Engine/Plugins/VirtualProduction/GeometryMask/Source/GeometryMask/Public/GeometryMaskCanvasResource.h
@@ -95,9 +95,8 @@ private:
 	UPROPERTY()
 	int32 NumUsedChannels = 0;
 
-	/** Underlying UCanvas object. */
-	UPROPERTY(Transient)
-	TObjectPtr<UCanvas> CanvasObject;
+	/** Canvas used for drawing. Using FCanvas over UCanvas for perf */
+	TSharedPtr<FCanvas> Canvas;
 
 	/** The default viewport size to use when it can't be resolved from the actual viewport. */
 	UPROPERTY(EditAnywhere, BlueprintReadWrite, Getter = "Auto", Category = "Rendering", meta = (AllowPrivateAccess = "true"))
