Geometry Mask: changed the canvas render target texture to be R16f instead of RGBA16f. This is to simplify the system by having one canvas render target be handled by one geometry mask canvas while also not wasting memory when a non-multiple of 3 number of masks are used in a scene.

#virtualized
#Synced-CL 49036542

diff --git a/Engine/Plugins/VirtualProduction/Avalanche/Source/AvalancheMask/Private/AvaMaskTypes.cpp b/Engine/Plugins/VirtualProduction/Avalanche/Source/AvalancheMask/Private/AvaMaskTypes.cpp
--- a/Engine/Plugins/VirtualProduction/Avalanche/Source/AvalancheMask/Private/AvaMaskTypes.cpp
+++ b/Engine/Plugins/VirtualProduction/Avalanche/Source/AvalancheMask/Private/AvaMaskTypes.cpp
@@ -15,7 +15,6 @@ bool FAvaMask2DMaterialParameters::HasSameParameters(const FAvaMask2DMaterialPar
 	return BlendMode == InOther.BlendMode
 		&& Texture == InOther.Texture
 		&& FMath::IsNearlyEqual(BaseOpacity, InOther.BaseOpacity, Tolerance)
-		&& ChannelAsVector.Equals(InOther.ChannelAsVector, Tolerance)
 		&& bInvert == InOther.bInvert
 		&& Padding.Equals(InOther.Padding, Tolerance)
 		&& bApplyFeathering == InOther.bApplyFeathering
@@ -36,7 +35,6 @@ bool FAvaMask2DMaterialParameters::ApplyToMID(UMaterialInstanceDynamic* InMateri
 
 	InMaterial->SetTextureParameterValueByInfo(UE::AvaMask::Internal::TextureParameterInfo, Texture);
 	InMaterial->SetScalarParameterValueByInfo(UE::AvaMask::Internal::BaseOpacityParameterInfo, FMath::Clamp(BaseOpacity, 0.0f, 1.0f));
-	InMaterial->SetVectorParameterValueByInfo(UE::AvaMask::Internal::ChannelParameterInfo, ChannelAsVector);
 	InMaterial->SetScalarParameterValueByInfo(UE::AvaMask::Internal::InvertParameterInfo, bInvert ? 1.f : 0.f);
 	InMaterial->SetVectorParameterValueByInfo(UE::AvaMask::Internal::PaddingParameterInfo, FLinearColor(FMath::Max(0, Padding.X), FMath::Max(0, Padding.Y), 0, 0));
 	InMaterial->SetVectorParameterValueByInfo(UE::AvaMask::Internal::FeatherParameterInfo, FLinearColor(bApplyFeathering ? 1.0f : 0.0f, FMath::Max(0, OuterFeatherRadius), FMath::Max(0, InnerFeatherRadius), FMath::Max(0, FMath::Max(OuterFeatherRadius, InnerFeatherRadius))));
@@ -47,21 +45,29 @@ bool FAvaMask2DMaterialParameters::StoreFromMaterial(const UMaterialInterface* I
 {
 	return SetFromMaterial(InMaterial
 		, UE::AvaMask::Internal::TextureParameterInfo
-		, UE::AvaMask::Internal::ChannelParameterInfo
 		, UE::AvaMask::Internal::InvertParameterInfo
 		, UE::AvaMask::Internal::BaseOpacityParameterInfo
 		, UE::AvaMask::Internal::PaddingParameterInfo
 		, UE::AvaMask::Internal::FeatherParameterInfo);
 }
 
-bool FAvaMask2DMaterialParameters::SetFromMaterial(
-	const UMaterialInterface* InMaterial
+bool FAvaMask2DMaterialParameters::SetFromMaterial(const UMaterialInterface* InMaterial
 	, const FMaterialParameterInfo& InTextureParameter
 	, const FMaterialParameterInfo& InChannelParameter
 	, const FMaterialParameterInfo& InInvertParameter
 	, const FMaterialParameterInfo& InBaseOpacityParameter
 	, const FMaterialParameterInfo& InPaddingParameter
 	, const FMaterialParameterInfo& InFeatherParameter)
+{
+	return SetFromMaterial(InMaterial, InTextureParameter, InInvertParameter, InBaseOpacityParameter, InPaddingParameter, InFeatherParameter);
+}
+
+bool FAvaMask2DMaterialParameters::SetFromMaterial(const UMaterialInterface* InMaterial
+	, const FMaterialParameterInfo& InTextureParameter
+	, const FMaterialParameterInfo& InInvertParameter
+	, const FMaterialParameterInfo& InBaseOpacityParameter
+	, const FMaterialParameterInfo& InPaddingParameter
+	, const FMaterialParameterInfo& InFeatherParameter)
 {
 	if (!ensureAlways(InMaterial))
 	{
@@ -97,18 +103,6 @@ bool FAvaMask2DMaterialParameters::SetFromMaterial(
 		BaseOpacity = FoundValue;
 	}
 
-	// Channel
-	{
-		FLinearColor FoundValue;
-		if (!InMaterial->GetVectorParameterValue(InChannelParameter, FoundValue))
-		{
-			LogMissingParameter(InTextureParameter);
-			return false;
-		}
-		ChannelAsVector = FoundValue;
-		Channel = UE::GeometryMask::VectorToMaskChannel(ChannelAsVector);
-	}
-
 	// Invert
 	{
 		float FoundValue = 0.0f;
diff --git a/Engine/Plugins/VirtualProduction/Avalanche/Source/AvalancheMask/Private/AvaMaskUtilities.h b/Engine/Plugins/VirtualProduction/Avalanche/Source/AvalancheMask/Private/AvaMaskUtilities.h
--- a/Engine/Plugins/VirtualProduction/Avalanche/Source/AvalancheMask/Private/AvaMaskUtilities.h
+++ b/Engine/Plugins/VirtualProduction/Avalanche/Source/AvalancheMask/Private/AvaMaskUtilities.h
@@ -18,26 +18,9 @@ namespace UE::AvaMask::Internal
 	/** Note that this assumes the parameters have the same name for all materials, which they should if using the provided material function. */
 	static FMaterialParameterInfo TextureParameterInfo = { TEXT("Mask_Texture") };
 	static FMaterialParameterInfo BaseOpacityParameterInfo = { TEXT("Mask_BaseOpacity") };
-	static FMaterialParameterInfo ChannelParameterInfo = { TEXT("Mask_Channel") };
 	static FMaterialParameterInfo InvertParameterInfo = { TEXT("Mask_Invert") };
 	static FMaterialParameterInfo PaddingParameterInfo = { TEXT("Mask_Padding") };
 	static FMaterialParameterInfo FeatherParameterInfo = { TEXT("Mask_Feather") };
-	
-	static TMap<FLinearColor, EGeometryMaskColorChannel> VectorToMaskChannelEnum =
-	{
-		{ {1, 0, 0, 0}, EGeometryMaskColorChannel::Red },
-		{ {0, 1, 0, 0}, EGeometryMaskColorChannel::Green },
-		{ {0, 0, 1, 0}, EGeometryMaskColorChannel::Blue },
-		{ {0, 0, 0, 1}, EGeometryMaskColorChannel::Alpha }
-	};
-	
-	static TMap<EGeometryMaskColorChannel, FLinearColor> MaskChannelEnumToVector =
-	{
-		{ EGeometryMaskColorChannel::Red, {1, 0, 0, 0} },
-		{ EGeometryMaskColorChannel::Green, {0, 1, 0, 0} },
-		{ EGeometryMaskColorChannel::Blue, {0, 0, 1, 0} },
-		{ EGeometryMaskColorChannel::Alpha, {0, 0, 0, 1} }
-	};
 
 	static FName HandleTag = TEXT("Modifier.Mask2D");
 
diff --git a/Engine/Plugins/VirtualProduction/Avalanche/Source/AvalancheMask/Private/Handling/AvaMaskDesignedMaterialHandle.cpp b/Engine/Plugins/VirtualProduction/Avalanche/Source/AvalancheMask/Private/Handling/AvaMaskDesignedMaterialHandle.cpp
--- a/Engine/Plugins/VirtualProduction/Avalanche/Source/AvalancheMask/Private/Handling/AvaMaskDesignedMaterialHandle.cpp
+++ b/Engine/Plugins/VirtualProduction/Avalanche/Source/AvalancheMask/Private/Handling/AvaMaskDesignedMaterialHandle.cpp
@@ -57,7 +57,6 @@ bool FAvaMaskDesignedMaterialHandle::GetMaskParameters(FAvaMask2DMaterialParamet
 		OutParameters.SetFromMaterial(
 			MaterialInstance
 			, UE::AvaMask::Internal::TextureParameterInfo
-			, UE::AvaMask::Internal::ChannelParameterInfo
 			, UE::AvaMask::Internal::InvertParameterInfo
 			, UE::AvaMask::Internal::BaseOpacityParameterInfo
 			, UE::AvaMask::Internal::PaddingParameterInfo
@@ -71,7 +70,6 @@ bool FAvaMaskDesignedMaterialHandle::GetMaskParameters(FAvaMask2DMaterialParamet
 bool FAvaMaskDesignedMaterialHandle::SetMaskParameters(
 	UTexture* InTexture
 	, float InBaseOpacity
-	, EGeometryMaskColorChannel InChannel
 	, bool bInInverted
 	, const FVector2f& InPadding
 	, bool bInApplyFeathering
@@ -80,7 +78,7 @@ bool FAvaMaskDesignedMaterialHandle::SetMaskParameters(
 {
 	if (UDynamicMaterialInstance* MaterialInstance = WeakDesignedMaterial.Get())
 	{
-		UE_LOG(LogAvaMask, VeryVerbose, TEXT("SetParameters: Texture:%s, Channel:%s"), InTexture ? *InTexture->GetName() : TEXT("(None)"), *UE::AvaMask::Internal::MaskChannelEnumToVector[InChannel].ToString());
+		UE_LOG(LogAvaMask, VeryVerbose, TEXT("SetParameters: Texture:%s"), InTexture ? *InTexture->GetName() : TEXT("(None)"));
 
 		UTexture* DefaultTexture = nullptr;
 		MaterialInstance->GetTextureParameterDefaultValue(UE::AvaMask::Internal::TextureParameterInfo, DefaultTexture);
@@ -88,7 +86,6 @@ bool FAvaMaskDesignedMaterialHandle::SetMaskParameters(
 		InTexture = InTexture ? InTexture : DefaultTexture;
 
 		MaterialInstance->SetTextureParameterValueByInfo(UE::AvaMask::Internal::TextureParameterInfo, InTexture);
-		MaterialInstance->SetVectorParameterValueByInfo(UE::AvaMask::Internal::ChannelParameterInfo, UE::AvaMask::Internal::MaskChannelEnumToVector[InChannel]);
 		MaterialInstance->SetScalarParameterValueByInfo(UE::AvaMask::Internal::InvertParameterInfo, bInInverted ? 1.0f : 0.0f);
 		MaterialInstance->SetScalarParameterValueByInfo(UE::AvaMask::Internal::BaseOpacityParameterInfo, FMath::Clamp(InBaseOpacity, 0.0f, 1.0f));
 		MaterialInstance->SetVectorParameterValueByInfo(UE::AvaMask::Internal::PaddingParameterInfo, FLinearColor(FMath::Max(0, InPadding.X), FMath::Max(0, InPadding.Y), 0, 0));
@@ -158,7 +155,6 @@ bool FAvaMaskDesignedMaterialHandle::SaveOriginalState(const FStructView& InHand
 			MaterialHandleData->OriginalMaskMaterialParameters.SetFromMaterial(
 				DesignedMaterial
 				, UE::AvaMask::Internal::TextureParameterInfo
-				, UE::AvaMask::Internal::ChannelParameterInfo
 				, UE::AvaMask::Internal::InvertParameterInfo
 				, UE::AvaMask::Internal::BaseOpacityParameterInfo
 				, UE::AvaMask::Internal::PaddingParameterInfo
@@ -189,7 +185,6 @@ bool FAvaMaskDesignedMaterialHandle::ApplyOriginalState(const FStructView& InHan
 			SetMaskParameters(
 				MaterialHandleData->OriginalMaskMaterialParameters.Texture
 				, MaterialHandleData->OriginalMaskMaterialParameters.BaseOpacity
-				, MaterialHandleData->OriginalMaskMaterialParameters.Channel
 				, MaterialHandleData->OriginalMaskMaterialParameters.bInvert
 				, MaterialHandleData->OriginalMaskMaterialParameters.Padding
 				, MaterialHandleData->OriginalMaskMaterialParameters.bApplyFeathering
@@ -221,7 +216,6 @@ bool FAvaMaskDesignedMaterialHandle::ApplyModifiedState(
 			SetMaskParameters(
 				InModifiedParameters.MaterialParameters.Texture
 				, InModifiedParameters.MaterialParameters.BaseOpacity
-				, InModifiedParameters.MaterialParameters.Channel
 				, InModifiedParameters.MaterialParameters.bInvert
 				, InModifiedParameters.MaterialParameters.Padding
 				, InModifiedParameters.MaterialParameters.bApplyFeathering
@@ -284,7 +278,6 @@ void FAvaMaskDesignedMaterialHandle::OnMaterialBuilt(UDynamicMaterialModelBase*
 	SetMaskParameters(
 		LastAppliedParameters.MaterialParameters.Texture
 		, LastAppliedParameters.MaterialParameters.BaseOpacity
-		, LastAppliedParameters.MaterialParameters.Channel
 		, LastAppliedParameters.MaterialParameters.bInvert
 		, LastAppliedParameters.MaterialParameters.Padding
 		, LastAppliedParameters.MaterialParameters.bApplyFeathering
diff --git a/Engine/Plugins/VirtualProduction/Avalanche/Source/AvalancheMask/Private/Handling/AvaMaskDesignedMaterialHandle.h b/Engine/Plugins/VirtualProduction/Avalanche/Source/AvalancheMask/Private/Handling/AvaMaskDesignedMaterialHandle.h
--- a/Engine/Plugins/VirtualProduction/Avalanche/Source/AvalancheMask/Private/Handling/AvaMaskDesignedMaterialHandle.h
+++ b/Engine/Plugins/VirtualProduction/Avalanche/Source/AvalancheMask/Private/Handling/AvaMaskDesignedMaterialHandle.h
@@ -41,7 +41,7 @@ public:
 
 	// ~Begin IAvaMaskMaterialHandle
 	virtual bool GetMaskParameters(FAvaMask2DMaterialParameters& OutParameters) override;
-	virtual bool SetMaskParameters(UTexture* InTexture, float InBaseOpacity, EGeometryMaskColorChannel InChannel, bool bInInverted, const FVector2f& InPadding, bool bInApplyFeathering, float InOuterFeathering, float InInnerFeathering) override;
+	virtual bool SetMaskParameters(UTexture* InTexture, float InBaseOpacity, bool bInInverted, const FVector2f& InPadding, bool bInApplyFeathering, float InOuterFeathering, float InInnerFeathering) override;
 	virtual const EBlendMode GetBlendMode() override;
 	virtual void SetBlendMode(const EBlendMode InBlendMode) override;
 	virtual void SetChildMaterial(UMaterialInstanceDynamic* InChildMaterial) override { }
diff --git a/Engine/Plugins/VirtualProduction/Avalanche/Source/AvalancheMask/Private/Handling/AvaMaskMaterialInstanceHandle.cpp b/Engine/Plugins/VirtualProduction/Avalanche/Source/AvalancheMask/Private/Handling/AvaMaskMaterialInstanceHandle.cpp
--- a/Engine/Plugins/VirtualProduction/Avalanche/Source/AvalancheMask/Private/Handling/AvaMaskMaterialInstanceHandle.cpp
+++ b/Engine/Plugins/VirtualProduction/Avalanche/Source/AvalancheMask/Private/Handling/AvaMaskMaterialInstanceHandle.cpp
@@ -53,7 +53,6 @@ bool FAvaMaskMaterialInstanceHandle::SaveOriginalState(const FStructView& InHand
 			HandleData->OriginalMaskMaterialParameters.SetFromMaterial(
 				ParentMaterial
 				, UE::AvaMask::Internal::TextureParameterInfo
-				, UE::AvaMask::Internal::ChannelParameterInfo
 				, UE::AvaMask::Internal::InvertParameterInfo
 				, UE::AvaMask::Internal::BaseOpacityParameterInfo
 				, UE::AvaMask::Internal::PaddingParameterInfo
@@ -77,7 +76,6 @@ bool FAvaMaskMaterialInstanceHandle::ApplyOriginalState(
 		SetMaskParameters(
 			HandleData->OriginalMaskMaterialParameters.Texture
 			, HandleData->OriginalMaskMaterialParameters.BaseOpacity
-			, HandleData->OriginalMaskMaterialParameters.Channel
 			, HandleData->OriginalMaskMaterialParameters.bInvert);
 		
 		if (UMaterialInterface* OriginalMaterial = HandleData->OriginalMaterial.Get())
@@ -109,7 +107,6 @@ bool FAvaMaskMaterialInstanceHandle::ApplyModifiedState(
 			SetMaskParameters(
 				InModifiedParameters.MaterialParameters.Texture
 				, InModifiedParameters.MaterialParameters.BaseOpacity
-				, InModifiedParameters.MaterialParameters.Channel
 				, InModifiedParameters.MaterialParameters.bInvert
 			);
 
diff --git a/Engine/Plugins/VirtualProduction/Avalanche/Source/AvalancheMask/Private/Handling/AvaMaskParametricMaterialHandle.cpp b/Engine/Plugins/VirtualProduction/Avalanche/Source/AvalancheMask/Private/Handling/AvaMaskParametricMaterialHandle.cpp
--- a/Engine/Plugins/VirtualProduction/Avalanche/Source/AvalancheMask/Private/Handling/AvaMaskParametricMaterialHandle.cpp
+++ b/Engine/Plugins/VirtualProduction/Avalanche/Source/AvalancheMask/Private/Handling/AvaMaskParametricMaterialHandle.cpp
@@ -38,7 +38,6 @@ bool FAvaMaskParametricMaterialHandle::GetMaskParameters(FAvaMask2DMaterialParam
 		OutParameters.SetFromMaterial(
 			MaterialInstance
 			, UE::AvaMask::Internal::TextureParameterInfo
-			, UE::AvaMask::Internal::ChannelParameterInfo
 			, UE::AvaMask::Internal::InvertParameterInfo
 			, UE::AvaMask::Internal::BaseOpacityParameterInfo
 			, UE::AvaMask::Internal::PaddingParameterInfo
@@ -52,7 +51,6 @@ bool FAvaMaskParametricMaterialHandle::GetMaskParameters(FAvaMask2DMaterialParam
 bool FAvaMaskParametricMaterialHandle::SetMaskParameters(
 	UTexture* InTexture
 	, float InBaseOpacity
-	, EGeometryMaskColorChannel InChannel
 	, bool bInInverted
 	, const FVector2f& InPadding
 	, bool bInApplyFeathering
@@ -64,10 +62,7 @@ bool FAvaMaskParametricMaterialHandle::SetMaskParameters(
 		FLinearColor PaddingValueV = FLinearColor(FMath::Max(0, InPadding.X), FMath::Max(0, InPadding.Y), 0, 0);
 		FLinearColor FeatherValueV = FLinearColor(bInApplyFeathering ? 1.0f : 0.0f, FMath::Max(0, InOuterFeathering), FMath::Max(0, InInnerFeathering), FMath::Max(0, FMath::Max(InOuterFeathering, InInnerFeathering)));
 		
-		UE_LOG(LogAvaMask, VeryVerbose, TEXT("SetParameters:\nTexture:%s\nChannel:%s\nFeather:%s\n"),
-			InTexture ? *InTexture->GetName() : TEXT("(None)"),
-			*UE::AvaMask::Internal::MaskChannelEnumToVector[InChannel].ToString(),
-			*FeatherValueV.ToString());
+		UE_LOG(LogAvaMask, VeryVerbose, TEXT("SetParameters:\nTexture:%s\n\nFeather:%s\n"), *GetNameSafe(InTexture), *FeatherValueV.ToString());
 
 		UTexture* DefaultTexture = nullptr;
 		MaterialInstance->GetTextureParameterDefaultValue(UE::AvaMask::Internal::TextureParameterInfo, DefaultTexture);
@@ -75,7 +70,6 @@ bool FAvaMaskParametricMaterialHandle::SetMaskParameters(
 		InTexture = InTexture ? InTexture : DefaultTexture;
 
 		MaterialInstance->SetTextureParameterValueByInfo(UE::AvaMask::Internal::TextureParameterInfo, InTexture);
-		MaterialInstance->SetVectorParameterValueByInfo(UE::AvaMask::Internal::ChannelParameterInfo, UE::AvaMask::Internal::MaskChannelEnumToVector[InChannel]);
 		MaterialInstance->SetScalarParameterValueByInfo(UE::AvaMask::Internal::InvertParameterInfo, bInInverted ? 1.0f : 0.0f);
 		MaterialInstance->SetScalarParameterValueByInfo(UE::AvaMask::Internal::BaseOpacityParameterInfo, FMath::Clamp(InBaseOpacity, 0.0f, 1.0f));
 		MaterialInstance->SetVectorParameterValueByInfo(UE::AvaMask::Internal::PaddingParameterInfo, PaddingValueV);
@@ -151,7 +145,6 @@ bool FAvaMaskParametricMaterialHandle::SaveOriginalState(const FStructView& InHa
 			MaterialHandleData->OriginalMaskMaterialParameters.SetFromMaterial(
 				Material
 				, UE::AvaMask::Internal::TextureParameterInfo
-				, UE::AvaMask::Internal::ChannelParameterInfo
 				, UE::AvaMask::Internal::InvertParameterInfo
 				, UE::AvaMask::Internal::BaseOpacityParameterInfo
 				, UE::AvaMask::Internal::PaddingParameterInfo
@@ -180,7 +173,6 @@ bool FAvaMaskParametricMaterialHandle::ApplyOriginalState(
 			SetMaskParameters(
 				MaterialHandleData->OriginalMaskMaterialParameters.Texture
 				, MaterialHandleData->OriginalMaskMaterialParameters.BaseOpacity
-				, MaterialHandleData->OriginalMaskMaterialParameters.Channel
 				, MaterialHandleData->OriginalMaskMaterialParameters.bInvert
 				, MaterialHandleData->OriginalMaskMaterialParameters.Padding
 				, MaterialHandleData->OriginalMaskMaterialParameters.bApplyFeathering
@@ -213,7 +205,6 @@ bool FAvaMaskParametricMaterialHandle::ApplyModifiedState(
 			SetMaskParameters(
 				InModifiedParameters.MaterialParameters.Texture
 				, InModifiedParameters.MaterialParameters.BaseOpacity
-				, InModifiedParameters.MaterialParameters.Channel
 				, InModifiedParameters.MaterialParameters.bInvert
 				, InModifiedParameters.MaterialParameters.Padding
 				, InModifiedParameters.MaterialParameters.bApplyFeathering
diff --git a/Engine/Plugins/VirtualProduction/Avalanche/Source/AvalancheMask/Private/Handling/AvaMaskParametricMaterialHandle.h b/Engine/Plugins/VirtualProduction/Avalanche/Source/AvalancheMask/Private/Handling/AvaMaskParametricMaterialHandle.h
--- a/Engine/Plugins/VirtualProduction/Avalanche/Source/AvalancheMask/Private/Handling/AvaMaskParametricMaterialHandle.h
+++ b/Engine/Plugins/VirtualProduction/Avalanche/Source/AvalancheMask/Private/Handling/AvaMaskParametricMaterialHandle.h
@@ -37,7 +37,7 @@ public:
 
 	// ~Begin IAvaMaskMaterialHandle
 	virtual bool GetMaskParameters(FAvaMask2DMaterialParameters& OutParameters) override;
-	virtual bool SetMaskParameters(UTexture* InTexture, float InBaseOpacity, EGeometryMaskColorChannel InChannel, bool bInInverted, const FVector2f& InPadding, bool bInApplyFeathering, float InOuterFeathering, float InInnerFeathering) override;
+	virtual bool SetMaskParameters(UTexture* InTexture, float InBaseOpacity, bool bInInverted, const FVector2f& InPadding, bool bInApplyFeathering, float InOuterFeathering, float InInnerFeathering) override;
 	virtual const EBlendMode GetBlendMode() override;
 	virtual void SetBlendMode(const EBlendMode InBlendMode) override;
 	virtual void SetChildMaterial(UMaterialInstanceDynamic* InChildMaterial) override;
diff --git a/Engine/Plugins/VirtualProduction/Avalanche/Source/AvalancheMask/Private/Handling/IAvaMaskMaterialHandle.h b/Engine/Plugins/VirtualProduction/Avalanche/Source/AvalancheMask/Private/Handling/IAvaMaskMaterialHandle.h
--- a/Engine/Plugins/VirtualProduction/Avalanche/Source/AvalancheMask/Private/Handling/IAvaMaskMaterialHandle.h
+++ b/Engine/Plugins/VirtualProduction/Avalanche/Source/AvalancheMask/Private/Handling/IAvaMaskMaterialHandle.h
@@ -33,7 +33,7 @@ public:
 	virtual const UScriptStruct* GetDataStructType() = 0;
 
 	virtual bool GetMaskParameters(FAvaMask2DMaterialParameters& OutParameters) = 0;
-	virtual bool SetMaskParameters(UTexture* InTexture, float InBaseOpacity, EGeometryMaskColorChannel InChannel, bool bInInverted = false, const FVector2f& InPadding = FVector2f::Zero(), bool bInApplyFeathering = false, float InOuterFeathering = 0.0f, float InInnerFeathering = 0.0f) = 0;
+	virtual bool SetMaskParameters(UTexture* InTexture, float InBaseOpacity, bool bInInverted = false, const FVector2f& InPadding = FVector2f::Zero(), bool bInApplyFeathering = false, float InOuterFeathering = 0.0f, float InInnerFeathering = 0.0f) = 0;
 
 	virtual const EBlendMode GetBlendMode() = 0;
 	virtual void SetBlendMode(const EBlendMode InBlendMode) = 0;
@@ -71,7 +71,7 @@ public:
 	virtual const UScriptStruct* GetDataStructType() override;
 
 	virtual bool GetMaskParameters(FAvaMask2DMaterialParameters& OutParameters) override;
-	virtual bool SetMaskParameters(UTexture* InTexture, float InBaseOpacity, EGeometryMaskColorChannel InChannel, bool bInInverted = false, const FVector2f& InPadding = FVector2f::Zero(), bool bInApplyFeathering = false, float InOuterFeathering = 0.0f, float InInnerFeathering = 0.0f) override;
+	virtual bool SetMaskParameters(UTexture* InTexture, float InBaseOpacity, bool bInInverted = false, const FVector2f& InPadding = FVector2f::Zero(), bool bInApplyFeathering = false, float InOuterFeathering = 0.0f, float InInnerFeathering = 0.0f) override;
 
 	virtual bool HasRequiredParameters(TArray<FString>& OutMissingParameterNames) override;
 
@@ -100,7 +100,6 @@ bool TAvaMaskMaterialHandle<HandleDataType>::GetMaskParameters(FAvaMask2DMateria
 		OutParameters.SetFromMaterial(
 			MaterialInstance
 			, UE::AvaMask::Internal::TextureParameterInfo
-			, UE::AvaMask::Internal::ChannelParameterInfo
 			, UE::AvaMask::Internal::InvertParameterInfo
 			, UE::AvaMask::Internal::BaseOpacityParameterInfo
 			, UE::AvaMask::Internal::PaddingParameterInfo
@@ -114,7 +113,6 @@ bool TAvaMaskMaterialHandle<HandleDataType>::GetMaskParameters(FAvaMask2DMateria
 template <typename HandleDataType>
 bool TAvaMaskMaterialHandle<HandleDataType>::SetMaskParameters(UTexture* InTexture
 	, float InBaseOpacity
-	, EGeometryMaskColorChannel InChannel
 	, bool bInInverted
 	, const FVector2f& InPadding
 	, bool bInApplyFeathering
@@ -123,7 +121,7 @@ bool TAvaMaskMaterialHandle<HandleDataType>::SetMaskParameters(UTexture* InTextu
 {
 	if (UMaterialInstanceDynamic* MaterialInstance = GetMaterialInstance())
 	{
-		UE_LOG(LogAvaMask, Verbose, TEXT("SetParameters: Texture:%s, Channel:%s, Material:%s"), InTexture ? *InTexture->GetName() : TEXT("(None)"), *UE::AvaMask::Internal::MaskChannelEnumToVector[InChannel].ToString(), *MaterialInstance->GetName());
+		UE_LOG(LogAvaMask, Verbose, TEXT("SetParameters: Texture:%s, Material:%s"), InTexture ? *InTexture->GetName() : TEXT("(None)"), *MaterialInstance->GetName());
 
 		UTexture* DefaultTexture = nullptr;
 		MaterialInstance->GetTextureParameterDefaultValue(UE::AvaMask::Internal::TextureParameterInfo, DefaultTexture);
@@ -131,7 +129,6 @@ bool TAvaMaskMaterialHandle<HandleDataType>::SetMaskParameters(UTexture* InTextu
 		InTexture = InTexture ? InTexture : DefaultTexture;
 
 		MaterialInstance->SetTextureParameterValueByInfo(UE::AvaMask::Internal::TextureParameterInfo, InTexture);
-		MaterialInstance->SetVectorParameterValueByInfo(UE::AvaMask::Internal::ChannelParameterInfo, UE::AvaMask::Internal::MaskChannelEnumToVector[InChannel]);
 		MaterialInstance->SetScalarParameterValueByInfo(UE::AvaMask::Internal::InvertParameterInfo, bInInverted ? 1.0f : 0.0f);
 		MaterialInstance->SetScalarParameterValueByInfo(UE::AvaMask::Internal::BaseOpacityParameterInfo, FMath::Clamp(InBaseOpacity, 0.0f, 1.0f));
 		MaterialInstance->SetVectorParameterValueByInfo(UE::AvaMask::Internal::PaddingParameterInfo, FLinearColor(FMath::Max(0, InPadding.X), FMath::Max(0, InPadding.Y), 0, 0));
@@ -153,7 +150,6 @@ bool TAvaMaskMaterialHandle<HandleDataType>::HasRequiredParameters(TArray<FStrin
 			, {
 				{ UE::AvaMask::Internal::TextureParameterInfo.Name, EMaterialParameterType::Texture }
 				, { UE::AvaMask::Internal::BaseOpacityParameterInfo.Name, EMaterialParameterType::Scalar }
-				, { UE::AvaMask::Internal::ChannelParameterInfo.Name, EMaterialParameterType::Vector }
 				, { UE::AvaMask::Internal::InvertParameterInfo.Name, EMaterialParameterType::Scalar }
 			}
 			, OutMissingParameterNames);
diff --git a/Engine/Plugins/VirtualProduction/Avalanche/Source/AvalancheMask/Private/Mask2D/AvaMask2DReadModifier.cpp b/Engine/Plugins/VirtualProduction/Avalanche/Source/AvalancheMask/Private/Mask2D/AvaMask2DReadModifier.cpp
--- a/Engine/Plugins/VirtualProduction/Avalanche/Source/AvalancheMask/Private/Mask2D/AvaMask2DReadModifier.cpp
+++ b/Engine/Plugins/VirtualProduction/Avalanche/Source/AvalancheMask/Private/Mask2D/AvaMask2DReadModifier.cpp
@@ -178,8 +178,6 @@ bool UAvaMask2DReadModifier::ApplyRead(TNotNull<UObject*> InMaterialContainer, T
 	MaskedState.MaterialParameters.CanvasName = InCurrentCanvas->GetCanvasName();
 	MaskedState.MaterialParameters.Texture = InCurrentCanvas->GetTexture();
 	MaskedState.MaterialParameters.BaseOpacity = BaseOpacity;
-	MaskedState.MaterialParameters.Channel = InCurrentCanvas->GetColorChannel();
-	MaskedState.MaterialParameters.ChannelAsVector = UE::GeometryMask::MaskChannelEnumToVector[InCurrentCanvas->GetColorChannel()];
 	MaskedState.MaterialParameters.bInvert = bInverted;
 	MaskedState.MaterialParameters.Padding = ViewportPadding;
 	MaskedState.MaterialParameters.bApplyFeathering = bUseFeathering;
diff --git a/Engine/Plugins/VirtualProduction/Avalanche/Source/AvalancheMask/Private/Mask2D/AvaMask2DWriteModifier.cpp b/Engine/Plugins/VirtualProduction/Avalanche/Source/AvalancheMask/Private/Mask2D/AvaMask2DWriteModifier.cpp
--- a/Engine/Plugins/VirtualProduction/Avalanche/Source/AvalancheMask/Private/Mask2D/AvaMask2DWriteModifier.cpp
+++ b/Engine/Plugins/VirtualProduction/Avalanche/Source/AvalancheMask/Private/Mask2D/AvaMask2DWriteModifier.cpp
@@ -84,7 +84,6 @@ bool UAvaMask2DWriteModifier::ApplyWrite(TNotNull<AActor*> InActor, TNotNull<UGe
 		{
 			FGeometryMaskWriteParameters Parameters = WriteComponent->GetParameters();
 			Parameters.CanvasName = Channel;
-			Parameters.ColorChannel = InCanvas->GetColorChannel();
 			Parameters.OperationType = WriteOperation;
 			WriteComponent->SetParameters(Parameters);
 		}
@@ -164,7 +163,6 @@ void UAvaMask2DWriteModifier::SetupMaskWriteComponent(IGeometryMaskWriteInterfac
 		// Set canvas name to write to
 		FGeometryMaskWriteParameters WriteParameters = InMaskWriter->GetParameters();
 		WriteParameters.CanvasName = Channel;
-		WriteParameters.ColorChannel = Canvas->GetColorChannel();
 
 		InMaskWriter->SetParameters(WriteParameters);
 	}
diff --git a/Engine/Plugins/VirtualProduction/Avalanche/Source/AvalancheMask/Public/AvaMaskTypes.h b/Engine/Plugins/VirtualProduction/Avalanche/Source/AvalancheMask/Public/AvaMaskTypes.h
--- a/Engine/Plugins/VirtualProduction/Avalanche/Source/AvalancheMask/Public/AvaMaskTypes.h
+++ b/Engine/Plugins/VirtualProduction/Avalanche/Source/AvalancheMask/Public/AvaMaskTypes.h
@@ -40,6 +40,14 @@ struct FAvaMask2DMaterialParameters
 {
 	GENERATED_BODY()
 
+PRAGMA_DISABLE_DEPRECATION_WARNINGS
+	FAvaMask2DMaterialParameters() = default;
+	FAvaMask2DMaterialParameters(const FAvaMask2DMaterialParameters&) = default;
+	FAvaMask2DMaterialParameters(FAvaMask2DMaterialParameters&&) = default;
+	FAvaMask2DMaterialParameters& operator=(const FAvaMask2DMaterialParameters&) = default;
+	FAvaMask2DMaterialParameters& operator=(FAvaMask2DMaterialParameters&&) = default;
+PRAGMA_ENABLE_DEPRECATION_WARNINGS
+
 	AVALANCHEMASK_API bool HasSameParameters(const FAvaMask2DMaterialParameters& InOther) const;
 
 	/** Canvas/Channel name. */
@@ -51,12 +59,14 @@ struct FAvaMask2DMaterialParameters
 	TObjectPtr<UTexture> Texture = nullptr;
 
 	/** Multiplies the mask texture to determine which channel to read from. */
+	UE_DEPRECATED(5.8, "Color channel management has been deprecated as it's no longer needed")
 	UPROPERTY()
-	EGeometryMaskColorChannel Channel = EGeometryMaskColorChannel::Red;
+	EGeometryMaskColorChannel Channel_DEPRECATED = EGeometryMaskColorChannel::Red;
 
 	/** Multiplies the mask texture to determine which channel to read from. */
+	UE_DEPRECATED(5.8, "Color channel management has been deprecated as it's no longer needed")
 	UPROPERTY()
-	FLinearColor ChannelAsVector = FLinearColor::Red;
+	FLinearColor ChannelAsVector_DEPRECATED = FLinearColor::Red;
 
 	/** Whether the mask result should be inverted when applied. This is stored in the material as a float to allow runtime modification. */
 	UPROPERTY()
@@ -86,7 +96,11 @@ struct FAvaMask2DMaterialParameters
 	bool StoreFromMaterial(const UMaterialInterface* InMaterial);
 
 	/** Set current values from the input material and named parameters, return true if successful. */
+	UE_DEPRECATED(5.8, "Color channel management has been deprecated as it's no longer needed. Use the overload without the Channel Parameter")
 	bool SetFromMaterial(const UMaterialInterface* InMaterial, const FMaterialParameterInfo& InTextureParameter, const FMaterialParameterInfo& InChannelParameter, const FMaterialParameterInfo& InInvertParameter, const FMaterialParameterInfo& InBaseOpacityParameter, const FMaterialParameterInfo& InPaddingParameter, const FMaterialParameterInfo& InFeatherParameter);
+
+	/** Set current values from the input material and named parameters, return true if successful. */
+	bool SetFromMaterial(const UMaterialInterface* InMaterial, const FMaterialParameterInfo& InTextureParameter, const FMaterialParameterInfo& InInvertParameter, const FMaterialParameterInfo& InBaseOpacityParameter, const FMaterialParameterInfo& InPaddingParameter, const FMaterialParameterInfo& InFeatherParameter);
 };
 
 /** Encapsulates all parameters to apply to a given modifier subject. */
diff --git a/Engine/Plugins/VirtualProduction/Avalanche/Source/AvalancheMaskEditor/Private/AvaMaskEditorMode.cpp b/Engine/Plugins/VirtualProduction/Avalanche/Source/AvalancheMaskEditor/Private/AvaMaskEditorMode.cpp
--- a/Engine/Plugins/VirtualProduction/Avalanche/Source/AvalancheMaskEditor/Private/AvaMaskEditorMode.cpp
+++ b/Engine/Plugins/VirtualProduction/Avalanche/Source/AvalancheMaskEditor/Private/AvaMaskEditorMode.cpp
@@ -271,8 +271,7 @@ void UAvaMaskEditorMode::OnSelectionChanged(const UTypedElementSelectionSet* InS
 	if (InSelectionSet)
 	{
 		PreviewCanvasId.ResetToNone();
-		PreviewCanvasChannel = EGeometryMaskColorChannel::None;
-		
+
 		if (const UTypedElementSelectionSet* SelectionSet = WeakActorSelectionSet.Get())
 		{
 			TArray<UObject*> AllSelected = SelectionSet->GetSelectedObjects();
@@ -300,7 +299,6 @@ void UAvaMaskEditorMode::OnSelectionChanged(const UTypedElementSelectionSet* InS
 					}
 
 					PreviewCanvasId = SelectedMaskCanvas->GetCanvasId();
-					PreviewCanvasChannel = SelectedMaskCanvas->GetColorChannel();
 				}
 			}
 			else
@@ -342,7 +340,6 @@ void UAvaMaskEditorMode::UpdatePreviewWidget()
 	if (CanvasPreviewWidget.IsValid())
 	{
 		CanvasPreviewWidget->SetCanvasId(PreviewCanvasId);
-		CanvasPreviewWidget->SetColorChannel(PreviewCanvasChannel);
 	}
 }
 
@@ -371,7 +368,6 @@ bool UAvaMaskEditorMode::AddMaskToSelected(const TArray<AActor*>& InMaskingActor
 
 	const UGeometryMaskCanvas* ParentCanvas = GetCanvasReferencedByActor(ParentActor);
 	const FName ChannelName = ParentMaskReadModifier->GetChannel();
-	const EGeometryMaskColorChannel ColorChannel = ParentCanvas ? ParentCanvas->GetColorChannel() : EGeometryMaskColorChannel::Red;
 
 	for (AActor* PlacedActor : InMaskingActors)
 	{
@@ -391,7 +387,6 @@ bool UAvaMaskEditorMode::AddMaskToSelected(const TArray<AActor*>& InMaskingActor
 			}
 
 			PreviewCanvasId = FGeometryMaskCanvasId(ParentActor->GetLevel(), ChannelName);
-			PreviewCanvasChannel = ColorChannel;
 		}
 	}
 		
diff --git a/Engine/Plugins/VirtualProduction/Avalanche/Source/AvalancheMaskEditor/Private/AvaMaskEditorMode.h b/Engine/Plugins/VirtualProduction/Avalanche/Source/AvalancheMaskEditor/Private/AvaMaskEditorMode.h
--- a/Engine/Plugins/VirtualProduction/Avalanche/Source/AvalancheMaskEditor/Private/AvaMaskEditorMode.h
+++ b/Engine/Plugins/VirtualProduction/Avalanche/Source/AvalancheMaskEditor/Private/AvaMaskEditorMode.h
@@ -59,8 +59,7 @@ private:
 
 	// @note: these can differ from the selection when interacting
 	FGeometryMaskCanvasId PreviewCanvasId;
-	EGeometryMaskColorChannel PreviewCanvasChannel = EGeometryMaskColorChannel::Red;
-	
+
 	AActor* GetActorToParentTo() const;
 
 	void OnActorSpawned(AActor* InActor);
diff --git a/Engine/Plugins/VirtualProduction/GeometryMask/Source/GeometryMask/Private/GeometryMaskCanvas.cpp b/Engine/Plugins/VirtualProduction/GeometryMask/Source/GeometryMask/Private/GeometryMaskCanvas.cpp
--- a/Engine/Plugins/VirtualProduction/GeometryMask/Source/GeometryMask/Private/GeometryMaskCanvas.cpp
+++ b/Engine/Plugins/VirtualProduction/GeometryMask/Source/GeometryMask/Private/GeometryMaskCanvas.cpp
@@ -16,13 +16,6 @@ const FName UGeometryMaskCanvas::ApplyFeatherPropertyName = GET_MEMBER_NAME_CHEC
 const FName UGeometryMaskCanvas::OuterFeatherRadiusPropertyName = GET_MEMBER_NAME_CHECKED(UGeometryMaskCanvas, InnerFeatherRadius);
 const FName UGeometryMaskCanvas::InnerFeatherRadiusPropertyName = GET_MEMBER_NAME_CHECKED(UGeometryMaskCanvas, OuterFeatherRadius);
 
-void UGeometryMaskCanvas::BeginDestroy()
-{
-	FreeResource();
-	
-	UObject::BeginDestroy();
-}
-
 #if WITH_EDITOR
 void UGeometryMaskCanvas::PostEditChangeProperty(FPropertyChangedEvent& InPropertyChangedEvent)
 {
@@ -123,7 +116,6 @@ bool UGeometryMaskCanvas::IsDefaultCanvas() const
 void UGeometryMaskCanvas::Free()
 {
 	Writers.Empty();
-	FreeResource();
 	OnDeactivated().ExecuteIfBound();
 }
 
@@ -214,7 +206,7 @@ void UGeometryMaskCanvas::SetInnerFeatherRadius(const int32 InValue)
 
 EGeometryMaskColorChannel UGeometryMaskCanvas::GetColorChannel() const
 {
-	return ColorChannel;
+	return EGeometryMaskColorChannel::Red;
 }
 
 void UGeometryMaskCanvas::RemoveInvalidWriters()
@@ -256,7 +248,7 @@ void UGeometryMaskCanvas::UpdateRenderParameters()
 {
 	if (CanvasResource)
 	{
-		CanvasResource->UpdateRenderParameters(ColorChannel, bApplyBlur, BlurStrength, bApplyFeather, OuterFeatherRadius, InnerFeatherRadius);
+		CanvasResource->UpdateRenderParameters(EGeometryMaskColorChannel::Red, bApplyBlur, BlurStrength, bApplyFeather, OuterFeatherRadius, InnerFeatherRadius);
 	}
 }
 
@@ -264,44 +256,23 @@ void UGeometryMaskCanvas::Initialize(const ULevel* InLevel, FName InCanvasName)
 {
 	CanvasName = InCanvasName;
 	CanvasId = FGeometryMaskCanvasId(InLevel, CanvasName);
+
+	CanvasResource = NewObject<UGeometryMaskCanvasResource>(this);
+	CanvasResource->OnDrawToCanvas().AddUObject(this, &UGeometryMaskCanvas::OnDrawToCanvas);
+	UpdateRenderParameters();
 }
 
-void UGeometryMaskCanvas::Update(
-	const ULevel* InLevel,
-	FSceneView& InView)
+void UGeometryMaskCanvas::Update(const ULevel* InLevel,FSceneView& InView)
 {
 	RemoveInvalidWriters();
 }
 
-void UGeometryMaskCanvas::AssignResource(
-	UGeometryMaskCanvasResource* InResource,
-	const EGeometryMaskColorChannel InColorChannel)
+void UGeometryMaskCanvas::AssignResource(UGeometryMaskCanvasResource* InResource, const EGeometryMaskColorChannel InColorChannel)
 {
-	// If we're already using this resource (just changing channel), no need to free
-	if (!CanvasResource
-		|| (CanvasResource && CanvasResource != InResource))
-	{
-		FreeResource();
-		CanvasResource = InResource;
-		CanvasResource->OnDrawToCanvas().AddUObject(this, &UGeometryMaskCanvas::OnDrawToCanvas);
-	}
-
-	ColorChannel = InColorChannel;
-
-	UpdateRenderParameters();
 }
 
 void UGeometryMaskCanvas::FreeResource()
 {
-	if (CanvasResource)
-	{
-		CanvasResource->OnDrawToCanvas().RemoveAll(this);
-		ColorChannel = EGeometryMaskColorChannel::None;
-	
-		// Free up the resource
-		CanvasResource->Checkin(CanvasId);
-		CanvasResource = nullptr;
-	}
 }
 
 FName UGeometryMaskCanvas::GetApplyBlurPropertyName()
diff --git a/Engine/Plugins/VirtualProduction/GeometryMask/Source/GeometryMask/Private/GeometryMaskCanvasResource.cpp b/Engine/Plugins/VirtualProduction/GeometryMask/Source/GeometryMask/Private/GeometryMaskCanvasResource.cpp
--- a/Engine/Plugins/VirtualProduction/GeometryMask/Source/GeometryMask/Private/GeometryMaskCanvasResource.cpp
+++ b/Engine/Plugins/VirtualProduction/GeometryMask/Source/GeometryMask/Private/GeometryMaskCanvasResource.cpp
@@ -38,18 +38,13 @@ namespace UE::GeometryMask::Private
 UGeometryMaskCanvasResource::UGeometryMaskCanvasResource()
 	: DefaultDrawingContext()
 {
-	// @note: We omit the Alpha channel as it's not uniformly supported
-	DependentCanvasIds = { { EGeometryMaskColorChannel::Red, FGeometryMaskCanvasId(EForceInit::ForceInit) }, { EGeometryMaskColorChannel::Green, FGeometryMaskCanvasId(EForceInit::ForceInit) }, { EGeometryMaskColorChannel::Blue, FGeometryMaskCanvasId(EForceInit::ForceInit) } };
-
-	UsedChannelMask.Init(false, MaxNumChannels);
-
 	FGeometryMaskPostProcessParameters_Blur PostProcessParameters_Blur;
-	PostProcessParameters_Blur.bPerChannelApplyBlur.SetRange(0, 4, false);
-	PostProcessParameters_Blur.PerChannelBlurStrength = { 16, 16, 16, 16 };
+	PostProcessParameters_Blur.bPerChannelApplyBlur.Init(false, UE::GeometryMask::MaxNumChannels);
+	PostProcessParameters_Blur.PerChannelBlurStrength.Init(16.0, UE::GeometryMask::MaxNumChannels);
 	PostProcess_Blur = MakeShared<FGeometryMaskPostProcess_Blur>(PostProcessParameters_Blur);
 
 	FGeometryMaskPostProcessParameters_DistanceField PostProcessParameters_DistanceField;
-	PostProcessParameters_DistanceField.bPerChannelCalculateDF.SetRange(0, 4, false);
+	PostProcessParameters_DistanceField.bPerChannelCalculateDF.Init(false, UE::GeometryMask::MaxNumChannels);
 	PostProcess_DistanceField = MakeShared<FGeometryMaskPostProcess_DistanceField>(PostProcessParameters_DistanceField);
 }
 
@@ -61,20 +56,10 @@ void UGeometryMaskCanvasResource::BeginDestroy()
 
 const EGeometryMaskColorChannel UGeometryMaskCanvasResource::GetNextAvailableColorChannel() const
 {
-	for (const TPair<EGeometryMaskColorChannel, FGeometryMaskCanvasId>& ColorChannelCanvas : DependentCanvasIds)
-	{
-		if (ColorChannelCanvas.Value.IsNone())
-		{
-			return ColorChannelCanvas.Key;
-		}
-	}
-
 	return EGeometryMaskColorChannel::None;
 }
 
-bool UGeometryMaskCanvasResource::Checkout(
-	const EGeometryMaskColorChannel InColorChannel,
-	const FGeometryMaskCanvasId& InRequestingCanvasId)
+bool UGeometryMaskCanvasResource::Checkout(const EGeometryMaskColorChannel InColorChannel,const FGeometryMaskCanvasId& InRequestingCanvasId)
 {
 	if (!ensure(InColorChannel != EGeometryMaskColorChannel::Num && InColorChannel != EGeometryMaskColorChannel::None))
 	{
@@ -82,16 +67,6 @@ bool UGeometryMaskCanvasResource::Checkout(
 		return false;
 	}
 
-	if (!ensure(DependentCanvasIds[InColorChannel].Name.IsNone()))
-	{
-		// Channel already in used
-		return false;
-	}
-
-	DependentCanvasIds[InColorChannel] = InRequestingCanvasId;
-	UsedChannelMask[static_cast<int32>(InColorChannel)] = true;
-	NumUsedChannels++;
-	
 	RemoveInvalidDrawingContexts();
 
 	FGeometryMaskDrawingContext* DrawingContext = GetDrawingContextForCanvas(InRequestingCanvasId);
@@ -106,117 +81,34 @@ bool UGeometryMaskCanvasResource::Checkout(
 
 bool UGeometryMaskCanvasResource::Checkin(const FGeometryMaskCanvasId& InRequestingCanvasId)
 {
-	for (TPair<EGeometryMaskColorChannel, FGeometryMaskCanvasId>& ColorChannelCanvas : DependentCanvasIds)
-	{
-		if (ColorChannelCanvas.Value == InRequestingCanvasId)
-		{
-			// Effectively free this ColorChannel and make available for Checkout
-			ColorChannelCanvas.Value.ResetToNone();
-			ResetRenderParameters(ColorChannelCanvas.Key);
-			UsedChannelMask[static_cast<int32>(ColorChannelCanvas.Key)] = false;
-			NumUsedChannels--;
-			
-			return true;
-		}
-	}
-
-	RemoveInvalidDrawingContexts();
-
 	// The requesting canvas name wasn't present in this Resource
+	RemoveInvalidDrawingContexts();
 	return false;
 }
 
 int32 UGeometryMaskCanvasResource::Compact()
 {
-	int32 NumUsedChannelsAfterCompact = 0;
-	EGeometryMaskColorChannel LastAvailableColorChannel = EGeometryMaskColorChannel::None;
-	for (int32 ChannelIdx = 0; ChannelIdx < MaxNumAvailableChannels; ++ChannelIdx)
-	{
-		EGeometryMaskColorChannel ColorChannel = static_cast<EGeometryMaskColorChannel>(ChannelIdx);
-		bool bIsChannelUsed = UsedChannelMask[ChannelIdx];
-
-		if (!bIsChannelUsed && LastAvailableColorChannel == EGeometryMaskColorChannel::None)
-		{
-			LastAvailableColorChannel = ColorChannel;
-			continue;
-		}
-
-		// Channel is used, and there was a previous one available to move to
-		if (bIsChannelUsed)
-		{
-			// Re-assign to last available channel
-			if (LastAvailableColorChannel != EGeometryMaskColorChannel::None)
-			{
-				FGeometryMaskCanvasId& UsedCanvasId = DependentCanvasIds[ColorChannel];
-				const ULevel* Level = UsedCanvasId.Level.ResolveObjectPtr();
-
-				if (Level && Level->OwningWorld)
-				{
-					if (UGeometryMaskWorldSubsystem* MaskSubsystem = Level->OwningWorld->GetSubsystem<UGeometryMaskWorldSubsystem>())
-					{
-						if (UGeometryMaskCanvas* UsedCanvas = MaskSubsystem->GetNamedCanvas(Level, UsedCanvasId.Name))
-						{
-							// Re-assign to that last available channel, freeing this one
-							UsedCanvas->AssignResource(this, LastAvailableColorChannel);						
-							++NumUsedChannelsAfterCompact;
-
-							// "Checkin" current channel and make available
-							UsedCanvasId.ResetToNone();
-							ResetRenderParameters(ColorChannel);
-							UsedChannelMask[ChannelIdx] = false;
-							LastAvailableColorChannel = ColorChannel;
-
-							// "Checkout" last channel that this was moved to
-							DependentCanvasIds[LastAvailableColorChannel] = UsedCanvasId;
-							UsedChannelMask[static_cast<int32>(LastAvailableColorChannel)] = true;
-							UsedCanvas->UpdateRenderParameters();
-						}
-					}
-				}
-			}
-			else
-			{
-				++NumUsedChannelsAfterCompact;
-			}
-		}
-	}
-
-	RemoveInvalidDrawingContexts();
-
-	NumUsedChannels = NumUsedChannelsAfterCompact;
-	return MaxNumAvailableChannels - NumUsedChannelsAfterCompact;
+	return 0;
 }
 
 int32 UGeometryMaskCanvasResource::GetNumChannelsUsed() const
 {
-	return NumUsedChannels;
+	return UE::GeometryMask::MaxNumChannels;
 }
 
 bool UGeometryMaskCanvasResource::IsAnyChannelUsed() const
 {
-	return NumUsedChannels > 0;
+	return true;
 }
 
 TArray<FGeometryMaskCanvasId> UGeometryMaskCanvasResource::GetDependentCanvasIds() const
 {
-	TArray<FGeometryMaskCanvasId> DependentCanvases;
-	DependentCanvases.Reserve(DependentCanvasIds.Num());
-
-	for (const TPair<EGeometryMaskColorChannel, FGeometryMaskCanvasId>& DependentCanvas : DependentCanvasIds)
-	{
-		if (!DependentCanvas.Value.IsNone())
-		{
-			DependentCanvases.Emplace(DependentCanvas.Value);
-		}
-	}
-	
-	return DependentCanvases;
+	return {};
 }
 
 void UGeometryMaskCanvasResource::UpdateViewportSize()
 {
-	if (!IsValid(RenderTargetTexture)
-		|| IsUnreachable())
+	if (!IsValid(RenderTargetTexture) || IsUnreachable())
 	{
 		return;
 	}
@@ -248,19 +140,19 @@ void UGeometryMaskCanvasResource::UpdateViewportSize()
 	int32 SizeX = NewScaledViewportSize.X;
 	int32 SizeY = NewScaledViewportSize.Y;
 
-	if (const float RatioX = static_cast<float>(SizeX) / MaxTextureSize;
+	if (const float RatioX = static_cast<float>(SizeX) / UE::GeometryMask::MaxTextureSize;
 		RatioX > 1.0)
 	{
 		// Width too big, cap to max and reduce height proportionally
-		SizeX = MaxTextureSize;
+		SizeX = UE::GeometryMask::MaxTextureSize;
 		SizeY /= RatioX;
 	}
 
-	if (const float RatioY = static_cast<float>(SizeY) / MaxTextureSize;
+	if (const float RatioY = static_cast<float>(SizeY) / UE::GeometryMask::MaxTextureSize;
 		RatioY > 1.0)
 	{
 		// Height too big, cap to max and reduce width proportionally
-		SizeY = MaxTextureSize;
+		SizeY = UE::GeometryMask::MaxTextureSize;
 		SizeX /= RatioY;
 	}
 
@@ -357,8 +249,7 @@ void UGeometryMaskCanvasResource::ResetRenderParameters(EGeometryMaskColorChanne
 
 void UGeometryMaskCanvasResource::SetViewportSize(FGeometryMaskDrawingContext& InDrawingContext, const FIntPoint& InViewportSize)
 {
-	if (InViewportSize.Size() > 0
-		&& InDrawingContext.ViewportSize != InViewportSize)
+	if (InViewportSize.Size() > 0 && InDrawingContext.ViewportSize != InViewportSize)
 	{
 		InDrawingContext.ViewportSize = InViewportSize;
 		UpdateViewportSize();
@@ -370,13 +261,8 @@ int32 UGeometryMaskCanvasResource::GetViewportPadding(const FGeometryMaskDrawing
 	int32 MaxFeatherRadius = 0;
 	{
 		const FGeometryMaskPostProcessParameters_DistanceField& DistanceFieldParameters = PostProcess_DistanceField->GetParameters();
-		for (int32 ChannelIdx = 0; ChannelIdx < MaxNumChannels; ++ChannelIdx)
+		for (int32 ChannelIdx = 0; ChannelIdx < UE::GeometryMask::MaxNumChannels; ++ChannelIdx)
 		{
-			if (!UsedChannelMask[ChannelIdx])
-			{
-				continue;
-			}
-			
 			if (DistanceFieldParameters.bPerChannelCalculateDF[ChannelIdx])
 			{
 				MaxFeatherRadius = FMath::Max(MaxFeatherRadius, DistanceFieldParameters.PerChannelRadius[ChannelIdx]);
@@ -387,13 +273,8 @@ int32 UGeometryMaskCanvasResource::GetViewportPadding(const FGeometryMaskDrawing
 	int32 MaxBlurRadius = 0;
 	{
 		const FGeometryMaskPostProcessParameters_Blur& BlurParameters = PostProcess_Blur->GetParameters();
-		for (int32 ChannelIdx = 0; ChannelIdx < MaxNumChannels; ++ChannelIdx)
+		for (int32 ChannelIdx = 0; ChannelIdx < UE::GeometryMask::MaxNumChannels; ++ChannelIdx)
 		{
-			if (!UsedChannelMask[ChannelIdx])
-			{
-				continue;
-			}
-			
 			if (BlurParameters.bPerChannelApplyBlur[ChannelIdx])
 			{
 				MaxBlurRadius = FMath::Max(MaxBlurRadius, UE::GeometryMask::Internal::ComputeEffectiveKernelSize(BlurParameters.PerChannelBlurStrength[ChannelIdx]));
@@ -412,6 +293,7 @@ UCanvasRenderTarget2D* UGeometryMaskCanvasResource::GetRenderTargetTexture()
 		RenderTargetTexture = NewObject<UCanvasRenderTarget2D>(this, ObjectName);
 		RenderTargetTexture->bForceLinearGamma = false;
 		RenderTargetTexture->bAutoGenerateMips = false;
+		RenderTargetTexture->RenderTargetFormat = RTF_R16f; // single channel only
 		RenderTargetTexture->SetSampleCount(GetSampleCount());
 		RenderTargetTexture->InitAutoFormat(MaxViewportSize.X, MaxViewportSize.Y);
 	}
@@ -419,10 +301,7 @@ UCanvasRenderTarget2D* UGeometryMaskCanvasResource::GetRenderTargetTexture()
 	return RenderTargetTexture;
 }
 
-void UGeometryMaskCanvasResource::Update(
-	const ULevel* InLevel,
-	FSceneView& InView,
-	int32 InViewIndex)
+void UGeometryMaskCanvasResource::Update(const ULevel* InLevel, FSceneView& InView, int32 InViewIndex)
 {
 	Draw(InLevel, InView, InViewIndex);
 }
@@ -540,7 +419,7 @@ FGeometryMaskDrawingContext* UGeometryMaskCanvasResource::GetDrawingContextForCa
 
 FGeometryMaskDrawingContext* UGeometryMaskCanvasResource::GetDrawingContextForChannel(EGeometryMaskColorChannel InColorChannel)
 {
-	return GetDrawingContextForCanvas(DependentCanvasIds[InColorChannel]);
+	return nullptr;
 }
 
 bool UGeometryMaskCanvasResource::RemoveInvalidDrawingContexts()
diff --git a/Engine/Plugins/VirtualProduction/GeometryMask/Source/GeometryMask/Private/GeometryMaskReadComponent.cpp b/Engine/Plugins/VirtualProduction/GeometryMask/Source/GeometryMask/Private/GeometryMaskReadComponent.cpp
--- a/Engine/Plugins/VirtualProduction/GeometryMask/Source/GeometryMask/Private/GeometryMaskReadComponent.cpp
+++ b/Engine/Plugins/VirtualProduction/GeometryMask/Source/GeometryMask/Private/GeometryMaskReadComponent.cpp
@@ -13,11 +13,5 @@ void UGeometryMaskReadComponent::SetParameters(FGeometryMaskReadParameters& InPa
 
 bool UGeometryMaskReadComponent::TryResolveCanvas()
 {
-	const bool bResult = TryResolveNamedCanvas(Parameters.CanvasName);
-	if (bResult)
-	{
-		Parameters.ColorChannel = CanvasWeak.Get()->GetColorChannel();
-	}
-	
-	return bResult;
+	return TryResolveNamedCanvas(Parameters.CanvasName);
 }
diff --git a/Engine/Plugins/VirtualProduction/GeometryMask/Source/GeometryMask/Private/GeometryMaskSubsystem.cpp b/Engine/Plugins/VirtualProduction/GeometryMask/Source/GeometryMask/Private/GeometryMaskSubsystem.cpp
--- a/Engine/Plugins/VirtualProduction/GeometryMask/Source/GeometryMask/Private/GeometryMaskSubsystem.cpp
+++ b/Engine/Plugins/VirtualProduction/GeometryMask/Source/GeometryMask/Private/GeometryMaskSubsystem.cpp
@@ -22,7 +22,6 @@ UGeometryMaskCanvas* UGeometryMaskSubsystem::GetDefaultCanvas()
 		const FName DefaultCanvasObjectName = FName(TEXT("GeometryMaskCanvas_Default"));
 		DefaultCanvas = NewObject<UGeometryMaskCanvas>(this, DefaultCanvasObjectName);
 		DefaultCanvas->Initialize(nullptr, FGeometryMaskCanvasId::DefaultCanvasName);
-		AssignResourceToCanvas(DefaultCanvas);
 	}
 	
 	return DefaultCanvas;
@@ -30,17 +29,19 @@ UGeometryMaskCanvas* UGeometryMaskSubsystem::GetDefaultCanvas()
 
 int32 UGeometryMaskSubsystem::GetNumCanvasResources() const
 {
-	return CanvasResources.Num();
+PRAGMA_DISABLE_DEPRECATION_WARNINGS
+	return CanvasResources_DEPRECATED.Num();
+PRAGMA_ENABLE_DEPRECATION_WARNINGS
 }
 
 const TSet<TObjectPtr<UGeometryMaskCanvasResource>>& UGeometryMaskSubsystem::GetCanvasResources() const
 {
-	return CanvasResources;
+PRAGMA_DISABLE_DEPRECATION_WARNINGS
+	return CanvasResources_DEPRECATED;
+PRAGMA_ENABLE_DEPRECATION_WARNINGS
 }
 
-void UGeometryMaskSubsystem::Update(
-	UWorld* InWorld,
-	FSceneViewFamily& InViewFamily)
+void UGeometryMaskSubsystem::Update(UWorld* InWorld, FSceneViewFamily& InViewFamily)
 {
 	if (!bDoUpdates)
 	{
@@ -110,79 +111,12 @@ void UGeometryMaskSubsystem::UpdateLevel(const ULevel* InLevel, UGeometryMaskWor
 
 void UGeometryMaskSubsystem::AssignResourceToCanvas(UGeometryMaskCanvas* InCanvas)
 {
-	UGeometryMaskCanvasResource* AvailableResource = nullptr;
-	EGeometryMaskColorChannel AvailableChannel = EGeometryMaskColorChannel::None;
-	for (UGeometryMaskCanvasResource* CanvasResource : CanvasResources)
-	{
-		if (AvailableChannel = CanvasResource->GetNextAvailableColorChannel();
-			AvailableChannel != EGeometryMaskColorChannel::None)
-		{
-			AvailableResource = CanvasResource;
-		}
-	}
-
-	// Nothing available, create new resource
-	if (AvailableChannel == EGeometryMaskColorChannel::None)
-	{
-		AvailableResource = NewObject<UGeometryMaskCanvasResource>(this);
-		CanvasResources.Emplace(AvailableResource);
-		AvailableChannel = AvailableResource->GetNextAvailableColorChannel();
-		OnGeometryMaskResourceCreatedDelegate.Broadcast(AvailableResource);
-	}
-
-	AvailableResource->Checkout(AvailableChannel, InCanvas->GetCanvasId());
-	InCanvas->AssignResource(AvailableResource, AvailableChannel);
 }
 
 void UGeometryMaskSubsystem::CompactResources()
 {
-	if (CanvasResources.Num() <= 1)
-	{
-		return;
-	}
-
-	TSet<TObjectPtr<UGeometryMaskCanvasResource>> UsedResources;
-	UsedResources.Reserve(CanvasResources.Num());
-	
-	for (UGeometryMaskCanvasResource* CanvasResource : CanvasResources)
-	{
-		CanvasResource->Compact();
-		if (CanvasResource->IsAnyChannelUsed())
-		{
-			UsedResources.Emplace(CanvasResource);
-		}
-		else
-		{
-			OnGeometryMaskResourceDestroyedDelegate.Broadcast(CanvasResource);
-		}
-	}
-
-	if (UsedResources.Num() != CanvasResources.Num())
-	{
-		CanvasResources = UsedResources;
-	}
 }
 
 void UGeometryMaskSubsystem::OnWorldDestroyed(UWorld* InWorld)
 {
-	CompactResources();
-	
-	TSet<TObjectPtr<UGeometryMaskCanvasResource>> UnusedResources;
-	UnusedResources.Reserve(CanvasResources.Num());
-	
-	for (UGeometryMaskCanvasResource* CanvasResource : CanvasResources)
-	{
-		if (!CanvasResource->IsAnyChannelUsed())
-		{
-			OnGeometryMaskResourceDestroyedDelegate.Broadcast(CanvasResource);
-			UnusedResources.Emplace(CanvasResource);
-		}
-	}
-
-	if (UnusedResources.IsEmpty())
-	{
-		return;
-	}
-
-	CanvasResources = CanvasResources.Difference(UnusedResources);
 }
diff --git a/Engine/Plugins/VirtualProduction/GeometryMask/Source/GeometryMask/Private/GeometryMaskWorldSubsystem.cpp b/Engine/Plugins/VirtualProduction/GeometryMask/Source/GeometryMask/Private/GeometryMaskWorldSubsystem.cpp
--- a/Engine/Plugins/VirtualProduction/GeometryMask/Source/GeometryMask/Private/GeometryMaskWorldSubsystem.cpp
+++ b/Engine/Plugins/VirtualProduction/GeometryMask/Source/GeometryMask/Private/GeometryMaskWorldSubsystem.cpp
@@ -25,19 +25,15 @@ void UGeometryMaskWorldSubsystem::Deinitialize()
 	{
 		for (const TPair<FName, TObjectPtr<UGeometryMaskCanvas>>& NamedCanvas : LevelState.Value.NamedCanvases)
 		{
-			if (IsValid(NamedCanvas.Value))
+			if (NamedCanvas.Value)
 			{
-				NamedCanvas.Value->Free();	
+				NamedCanvas.Value->Free();
+				OnGeometryMaskCanvasDestroyedDelegate.Broadcast(NamedCanvas.Value->GetCanvasId());
 			}
 		}
 	}
 
 	LevelStates.Empty();
-
-	if (UGeometryMaskSubsystem* EngineSubsystem = GEngine->GetEngineSubsystem<UGeometryMaskSubsystem>())
-	{
-		EngineSubsystem->OnWorldDestroyed(GetWorld());
-	}
 }
 
 const FGeometryMaskLevelState* UGeometryMaskWorldSubsystem::FindLevelState(const ULevel* InLevel) const
@@ -87,16 +83,25 @@ UGeometryMaskCanvas* UGeometryMaskWorldSubsystem::GetNamedCanvas(const ULevel* I
 	const TObjectPtr<UGeometryMaskCanvas>& NewCanvas = LevelState.NamedCanvases.Emplace(InName, NewObject<UGeometryMaskCanvas>(this, ObjectName));
 
 	NewCanvas->Initialize(InLevel, InName);
-	EngineSubsystem->AssignResourceToCanvas(NewCanvas);
-
-	NewCanvas->OnActivated().BindUObject(this, &UGeometryMaskWorldSubsystem::OnCanvasActivated, NewCanvas.Get());
-	NewCanvas->OnDeactivated().BindUObject(this, &UGeometryMaskWorldSubsystem::OnCanvasDeactivated, NewCanvas.Get());
-	
 	OnGeometryMaskCanvasCreatedDelegate.Broadcast(NewCanvas);
 
 	return NewCanvas;
 }
 
+void UGeometryMaskWorldSubsystem::ForEachCanvas(TFunctionRef<void(UGeometryMaskCanvas*, const FCanvasContext&)> InFunc) const
+{
+	FCanvasContext Context;
+	for (const TPair<TWeakObjectPtr<const ULevel>, FGeometryMaskLevelState>& LevelStatePair : LevelStates)
+	{
+		Context.LevelWeak = LevelStatePair.Key;
+		for (const TPair<FName, TObjectPtr<UGeometryMaskCanvas>>& CanvasPair : LevelStatePair.Value.NamedCanvases)
+		{
+			Context.CanvasName = CanvasPair.Key;
+			InFunc(CanvasPair.Value, Context);
+		}
+	}
+}
+
 TArray<FName> UGeometryMaskWorldSubsystem::GetCanvasNames(const ULevel* InLevel)
 {
 	TArray<FName> CanvasNames;
@@ -125,8 +130,7 @@ int32 UGeometryMaskWorldSubsystem::RemoveWithoutWriters()
 			{
 				if (NamedCanvas->GetWriters().IsEmpty())
 				{
-					OnGeometryMaskCanvasDestroyed().Broadcast(NamedCanvas->GetCanvasId());
-					NamedCanvas->FreeResource();
+					OnGeometryMaskCanvasDestroyedDelegate.Broadcast(NamedCanvas->GetCanvasId());
 				}
 
 				CanvasIter.RemoveCurrent();
@@ -142,43 +146,3 @@ int32 UGeometryMaskWorldSubsystem::RemoveWithoutWriters()
 
 	return NumRemoved;
 }
-
-void UGeometryMaskWorldSubsystem::OnCanvasActivated(UGeometryMaskCanvas* InCanvas)
-{
-	check(InCanvas);
-	
-	if (InCanvas->IsDefaultCanvas())
-	{
-		return;
-	}
-
-	// Already has a resource
-	if (InCanvas->GetResource())
-	{
-		return;
-	}
-
-	UGeometryMaskSubsystem* EngineSubsystem = GEngine->GetEngineSubsystem<UGeometryMaskSubsystem>();
-	if (!ensureAlwaysMsgf(EngineSubsystem, TEXT("UGeometryMaskSubsystem not resolved.")))
-	{
-		return;
-	}
-
-	// Provide a new resource for the canvas to write to
-	EngineSubsystem->AssignResourceToCanvas(InCanvas);
-}
-
-void UGeometryMaskWorldSubsystem::OnCanvasDeactivated(UGeometryMaskCanvas* InCanvas)
-{
-	if (!InCanvas || !IsValid(InCanvas) || InCanvas->IsDefaultCanvas())
-	{
-		return;
-	}
-
-	if (const UGeometryMaskCanvasResource* CanvasResource = InCanvas->GetResource())
-	{
-		// Resource assigned, so free it up
-		OnGeometryMaskCanvasDestroyed().Broadcast(InCanvas->GetCanvasId());
-		InCanvas->FreeResource();
-	}
-}
diff --git a/Engine/Plugins/VirtualProduction/GeometryMask/Source/GeometryMask/Private/GeometryMaskWriteComponent.cpp b/Engine/Plugins/VirtualProduction/GeometryMask/Source/GeometryMask/Private/GeometryMaskWriteComponent.cpp
--- a/Engine/Plugins/VirtualProduction/GeometryMask/Source/GeometryMask/Private/GeometryMaskWriteComponent.cpp
+++ b/Engine/Plugins/VirtualProduction/GeometryMask/Source/GeometryMask/Private/GeometryMaskWriteComponent.cpp
@@ -50,13 +50,11 @@ void UGeometryMaskWriteMeshComponent::DrawToCanvas(FCanvas* InCanvas)
 
 	UpdateCachedData();
 
-	FColor Color = UE::GeometryMask::MaskChannelEnumToColor[GetParameters().ColorChannel];
-	ESimpleElementBlendMode ElementBlendMode = SE_BLEND_Additive;
-	if (GetParameters().OperationType == EGeometryMaskCompositeOperation::Subtract)
-	{
-		Color = FColor(255 - Color.R, 255 - Color.G, 255 - Color.B, 255);
-		ElementBlendMode = SE_BLEND_Modulate;
-	}
+	const FLinearColor Color = GetParameters().OperationType == EGeometryMaskCompositeOperation::Subtract 
+		? FLinearColor::Black 
+		: FLinearColor::White;
+
+	constexpr ESimpleElementBlendMode ElementBlendMode = ESimpleElementBlendMode::SE_BLEND_Opaque;
 
 	// Write to Canvas BatchElements
 	{
@@ -440,7 +438,6 @@ bool UGeometryMaskWriteMeshComponent::TryResolveCanvas()
 		if (UGeometryMaskCanvas* Canvas = CanvasWeak.Get())
 		{
 			// Register this writer
-			Parameters.ColorChannel = Canvas->GetColorChannel();
 			Canvas->AddWriter(this);
 			return true;
 		}
diff --git a/Engine/Plugins/VirtualProduction/GeometryMask/Source/GeometryMask/Private/Shaders/GeometryMaskPostProcess_Blur.cpp b/Engine/Plugins/VirtualProduction/GeometryMask/Source/GeometryMask/Private/Shaders/GeometryMaskPostProcess_Blur.cpp
--- a/Engine/Plugins/VirtualProduction/GeometryMask/Source/GeometryMask/Private/Shaders/GeometryMaskPostProcess_Blur.cpp
+++ b/Engine/Plugins/VirtualProduction/GeometryMask/Source/GeometryMask/Private/Shaders/GeometryMaskPostProcess_Blur.cpp
@@ -140,9 +140,9 @@ namespace UE::GeometryMask
 }
 
 FGeometryMaskPostProcessParameters_Blur::FGeometryMaskPostProcessParameters_Blur()
-	: PerChannelBlurStrength({ 16, 16, 16, 16})
 {
-	bPerChannelApplyBlur.Init(false, 4);
+	bPerChannelApplyBlur.Init(false, UE::GeometryMask::MaxNumChannels);
+	PerChannelBlurStrength.Init(16.0,  UE::GeometryMask::MaxNumChannels);
 }
 
 FGeometryMaskPostProcess_Blur::FGeometryMaskPostProcess_Blur(const FGeometryMaskPostProcessParameters_Blur& InParameters)
@@ -176,11 +176,10 @@ void FGeometryMaskPostProcess_Blur::Execute_RenderThread(
 {
 	ensure(IsInRenderingThread());
 
-	static constexpr int32 ChannelCount = 4;
-	int32 PerChannelSampleCount[ChannelCount];
-	TArray<FVector4f> PerChannelWeightsAndOffsets[ChannelCount];
+	int32 PerChannelSampleCount[UE::GeometryMask::MaxNumChannels];
+	TArray<FVector4f> PerChannelWeightsAndOffsets[UE::GeometryMask::MaxNumChannels];
 
-	for (int32 ChannelIdx = 0; ChannelIdx < ChannelCount; ++ChannelIdx)
+	for (int32 ChannelIdx = 0; ChannelIdx < UE::GeometryMask::MaxNumChannels; ++ChannelIdx)
 	{
 		double BlurStrength = Parameters.PerChannelBlurStrength[ChannelIdx];
 
@@ -232,7 +231,7 @@ void FGeometryMaskPostProcess_Blur::Execute_RenderThread(
 			
 			FGeometryMaskBlurCS::FParameters* PassParameters = GraphBuilder.AllocParameters<FGeometryMaskBlurCS::FParameters>();
 			{
-				for (int32 ChannelIdx = 0; ChannelIdx < ChannelCount; ++ChannelIdx)
+				for (int32 ChannelIdx = 0; ChannelIdx < UE::GeometryMask::MaxNumChannels; ++ChannelIdx)
 				{
 					int32 WeightsAndOffsetsIdxOffset = ChannelIdx * UE::GeometryMask::Private::MaxNumSteps;
 					for (int32 WeightAndOffsetIdx = 0; WeightAndOffsetIdx < PerChannelWeightsAndOffsets[ChannelIdx].Num(); ++WeightAndOffsetIdx)
@@ -268,7 +267,7 @@ void FGeometryMaskPostProcess_Blur::Execute_RenderThread(
 			
 			FGeometryMaskBlurCS::FParameters* PassParameters = GraphBuilder.AllocParameters<FGeometryMaskBlurCS::FParameters>();
 			{
-				for (int32 ChannelIdx = 0; ChannelIdx < ChannelCount; ++ChannelIdx)
+				for (int32 ChannelIdx = 0; ChannelIdx < UE::GeometryMask::MaxNumChannels; ++ChannelIdx)
 				{
 					int32 WeightsAndOffsetsIdxOffset = ChannelIdx * UE::GeometryMask::Private::MaxNumSteps;
 					for (int32 WeightAndOffsetIdx = 0; WeightAndOffsetIdx < PerChannelWeightsAndOffsets[ChannelIdx].Num(); ++WeightAndOffsetIdx)
diff --git a/Engine/Plugins/VirtualProduction/GeometryMask/Source/GeometryMask/Private/Shaders/GeometryMaskPostProcess_Blur.h b/Engine/Plugins/VirtualProduction/GeometryMask/Source/GeometryMask/Private/Shaders/GeometryMaskPostProcess_Blur.h
--- a/Engine/Plugins/VirtualProduction/GeometryMask/Source/GeometryMask/Private/Shaders/GeometryMaskPostProcess_Blur.h
+++ b/Engine/Plugins/VirtualProduction/GeometryMask/Source/GeometryMask/Private/Shaders/GeometryMaskPostProcess_Blur.h
@@ -5,6 +5,7 @@
 #include "Containers/Array.h"
 #include "Containers/BitArray.h"
 #include "Containers/ContainerAllocationPolicies.h"
+#include "GeometryMaskConstants.h"
 #include "GeometryMaskPostProcess.h"
 #include "Templates/SharedPointer.h"
 
@@ -17,8 +18,8 @@ struct FGeometryMaskPostProcessParameters_Blur
 {
 	FGeometryMaskPostProcessParameters_Blur();
 	
-	TBitArray<TInlineAllocator<4>> bPerChannelApplyBlur;
-	TArray<double, TInlineAllocator<4>> PerChannelBlurStrength;
+	TBitArray<TInlineAllocator<UE::GeometryMask::MaxNumChannels>> bPerChannelApplyBlur;
+	TArray<double, TInlineAllocator<UE::GeometryMask::MaxNumChannels>> PerChannelBlurStrength;
 };
 
 class FGeometryMaskPostProcess_Blur
diff --git a/Engine/Plugins/VirtualProduction/GeometryMask/Source/GeometryMask/Private/Shaders/GeometryMaskPostProcess_DistanceField.cpp b/Engine/Plugins/VirtualProduction/GeometryMask/Source/GeometryMask/Private/Shaders/GeometryMaskPostProcess_DistanceField.cpp
--- a/Engine/Plugins/VirtualProduction/GeometryMask/Source/GeometryMask/Private/Shaders/GeometryMaskPostProcess_DistanceField.cpp
+++ b/Engine/Plugins/VirtualProduction/GeometryMask/Source/GeometryMask/Private/Shaders/GeometryMaskPostProcess_DistanceField.cpp
@@ -15,7 +15,6 @@ DECLARE_GPU_STAT_NAMED(GeometryMaskJFtoDF, TEXT("GeometryMaskJFtoDF"));
 
 namespace UE::GeometryMask::Private
 {
-	static constexpr int32 MaxNumChannels = 4;
 	static constexpr int32 NumNeighbors = 8;
 	static constexpr int32 MaxSteps = 13;
 
@@ -57,7 +56,7 @@ class FGeometryMaskJFInitCSBase : public FGlobalShader
 	SHADER_USE_PARAMETER_STRUCT(FGeometryMaskJFInitCSBase, FGlobalShader);
 
 	BEGIN_SHADER_PARAMETER_STRUCT(FParameters, )
-		SHADER_PARAMETER_ARRAY(FVector4f, Offsets, [UE::GeometryMask::Private::MaxNumChannels * (UE::GeometryMask::Private::NumNeighbors + 1)])
+		SHADER_PARAMETER_ARRAY(FVector4f, Offsets, [UE::GeometryMask::MaxNumChannels * (UE::GeometryMask::Private::NumNeighbors + 1)])
 		SHADER_PARAMETER(FIntPoint, InputDimensions)
 		SHADER_PARAMETER(FVector2f, OneOverInputDimensions)
 		SHADER_PARAMETER(FVector2f, UVRatioAdjustment)
@@ -207,8 +206,8 @@ IMPLEMENT_SHADER_TYPE(template<>, FGeometryMaskJFtoDFCS4, TEXT("/Plugin/Geometry
 
 FGeometryMaskPostProcessParameters_DistanceField::FGeometryMaskPostProcessParameters_DistanceField()
 {
-	bPerChannelCalculateDF.Init(false, 4);
-	PerChannelRadius.Init(0.0, 4);
+	bPerChannelCalculateDF.Init(false, UE::GeometryMask::MaxNumChannels);
+	PerChannelRadius.Init(0.0, UE::GeometryMask::MaxNumChannels);
 }
 
 FGeometryMaskPostProcess_DistanceField::FGeometryMaskPostProcess_DistanceField(const FGeometryMaskPostProcessParameters_DistanceField& InParameters)
@@ -265,7 +264,7 @@ void FGeometryMaskPostProcess_DistanceField::Execute_RenderThread(
 		FIntVector4 ChannelMap(ForceInitToZero);
 		{
 			int32 ActiveChannelIdx = 0;
-			for (int32 ChannelIdx = 0; ChannelIdx < UE::GeometryMask::Private::MaxNumChannels; ++ChannelIdx)
+			for (int32 ChannelIdx = 0; ChannelIdx < UE::GeometryMask::MaxNumChannels; ++ChannelIdx)
 			{
 				for (int32 ChannelMaskIdx = ActiveChannelIdx; ChannelMaskIdx < Parameters.bPerChannelCalculateDF.Num(); ++ChannelMaskIdx)
 				{
@@ -283,7 +282,7 @@ void FGeometryMaskPostProcess_DistanceField::Execute_RenderThread(
 
 		// Used to calculate step size per-channel
 		FIntVector4 SizeFromRadius;
-		for (int32 ChannelIdx = 0; ChannelIdx < UE::GeometryMask::Private::MaxNumChannels; ++ChannelIdx)
+		for (int32 ChannelIdx = 0; ChannelIdx < UE::GeometryMask::MaxNumChannels; ++ChannelIdx)
 		{
 			SizeFromRadius[ChannelIdx] = Parameters.PerChannelRadius[ChannelIdx];
 		}
@@ -324,19 +323,19 @@ void FGeometryMaskPostProcess_DistanceField::Execute_RenderThread(
 				StepComputeShader = StepComputeShader1;
 				OutputComputeShader = OutputComputeShader1;
 			}
-			else if (NumActiveChannels == 2)
+			else if (UE::GeometryMask::MaxNumChannels >= 2 && NumActiveChannels == 2)
 			{
 				InitComputeShader = InitComputeShader2;
 				StepComputeShader = StepComputeShader2;
 				OutputComputeShader = OutputComputeShader2;
 			}
-			else if (NumActiveChannels == 3)
+			else if (UE::GeometryMask::MaxNumChannels >= 3 && NumActiveChannels == 3)
 			{
 				InitComputeShader = InitComputeShader3;
 				StepComputeShader = StepComputeShader3;
 				OutputComputeShader = OutputComputeShader3;
 			}
-			else if(NumActiveChannels == 4)
+			else if(UE::GeometryMask::MaxNumChannels == 4 && NumActiveChannels == 4)
 			{
 				InitComputeShader = InitComputeShader4;
 				StepComputeShader = StepComputeShader4;
@@ -522,7 +521,7 @@ void FGeometryMaskPostProcess_DistanceField::Execute_RenderThread(
 				FGeometryMaskJFtoDFCSBase::FParameters* PassParameters = GraphBuilder.AllocParameters<FGeometryMaskJFtoDFCSBase::FParameters>();
 				{
 					FVector4f StepDistanceMultipliers;
-					for (int32 ChannelIdx = 0; ChannelIdx < UE::GeometryMask::Private::MaxNumChannels; ++ChannelIdx)
+					for (int32 ChannelIdx = 0; ChannelIdx < UE::GeometryMask::MaxNumChannels; ++ChannelIdx)
 					{
 						const float Radius = Parameters.bPerChannelCalculateDF[ChannelIdx] ? FMath::Max(1, Parameters.PerChannelRadius[ChannelIdx]) : InputSize.GetMax();
 						StepDistanceMultipliers[ChannelIdx] = static_cast<float>(InputSize.GetMax()) / Radius; 
diff --git a/Engine/Plugins/VirtualProduction/GeometryMask/Source/GeometryMask/Private/Shaders/GeometryMaskPostProcess_DistanceField.h b/Engine/Plugins/VirtualProduction/GeometryMask/Source/GeometryMask/Private/Shaders/GeometryMaskPostProcess_DistanceField.h
--- a/Engine/Plugins/VirtualProduction/GeometryMask/Source/GeometryMask/Private/Shaders/GeometryMaskPostProcess_DistanceField.h
+++ b/Engine/Plugins/VirtualProduction/GeometryMask/Source/GeometryMask/Private/Shaders/GeometryMaskPostProcess_DistanceField.h
@@ -5,6 +5,7 @@
 #include "Containers/Array.h"
 #include "Containers/BitArray.h"
 #include "Containers/ContainerAllocationPolicies.h"
+#include "GeometryMaskConstants.h"
 #include "GeometryMaskPostProcess.h"
 #include "RenderGraphResources.h"
 #include "Templates/SharedPointer.h"
@@ -19,8 +20,8 @@ struct FGeometryMaskPostProcessParameters_DistanceField
 {
 	FGeometryMaskPostProcessParameters_DistanceField();
 	
-	TBitArray<TInlineAllocator<4>> bPerChannelCalculateDF;
-	TArray<double, TInlineAllocator<4>> PerChannelRadius;
+	TBitArray<TInlineAllocator<UE::GeometryMask::MaxNumChannels>> bPerChannelCalculateDF;
+	TArray<double, TInlineAllocator<UE::GeometryMask::MaxNumChannels>> PerChannelRadius;
 };
 
 /** */
diff --git a/Engine/Plugins/VirtualProduction/GeometryMask/Source/GeometryMask/Public/GeometryMaskCanvas.h b/Engine/Plugins/VirtualProduction/GeometryMask/Source/GeometryMask/Public/GeometryMaskCanvas.h
--- a/Engine/Plugins/VirtualProduction/GeometryMask/Source/GeometryMask/Public/GeometryMaskCanvas.h
+++ b/Engine/Plugins/VirtualProduction/GeometryMask/Source/GeometryMask/Public/GeometryMaskCanvas.h
@@ -30,8 +30,6 @@ class GEOMETRYMASK_API UGeometryMaskCanvas : public UObject
 	GENERATED_BODY()
 
 public:
-	virtual void BeginDestroy() override;
-
 #if WITH_EDITOR
 	virtual void PostEditChangeProperty(FPropertyChangedEvent& InPropertyChangedEvent) override;
 #endif
@@ -97,7 +95,8 @@ public:
 	void SetInnerFeatherRadius(const int32 InValue);
 
 	/** Get the color channel to write to in the texture. */
-	UFUNCTION(BlueprintCallable, BlueprintPure, Category = "Rendering")
+	UE_DEPRECATED(5.8, "Color channel management has been deprecated as it's no longer needed")
+	UFUNCTION(BlueprintCallable, BlueprintPure, Category = "Rendering", meta=(DeprecatedFunction, DeprecationMessage="Color channel management has been deprecated as it's no longer needed"))
 	EGeometryMaskColorChannel GetColorChannel() const;
 
 	/** Called when Writers becomes non-empty. */
@@ -122,7 +121,10 @@ public:
 		return CanvasResource;
 	}
 
+	UE_DEPRECATED(5.8, "AssignResource has been deprecated as resource management is done by the canvas itself")
 	void AssignResource(UGeometryMaskCanvasResource* InResource, EGeometryMaskColorChannel InColorChannel);
+
+	UE_DEPRECATED(5.8, "FreeResource has been deprecated as resource management is done by the canvas itself")
 	void FreeResource();
 
 public:
@@ -189,10 +191,6 @@ private:
 	UPROPERTY(Transient)
 	TObjectPtr<UGeometryMaskCanvasResource> CanvasResource;
 
-	/** Color Channel being written to in the texture. */
-	UPROPERTY(Transient, Getter)
-	EGeometryMaskColorChannel ColorChannel = EGeometryMaskColorChannel::None;
-
 	/** List of objects that write to this canvas. */
 	TArray<TWeakInterfacePtr<IGeometryMaskWriteInterface>> Writers;
 };
diff --git a/Engine/Plugins/VirtualProduction/GeometryMask/Source/GeometryMask/Public/GeometryMaskCanvasResource.h b/Engine/Plugins/VirtualProduction/GeometryMask/Source/GeometryMask/Public/GeometryMaskCanvasResource.h
--- a/Engine/Plugins/VirtualProduction/GeometryMask/Source/GeometryMask/Public/GeometryMaskCanvasResource.h
+++ b/Engine/Plugins/VirtualProduction/GeometryMask/Source/GeometryMask/Public/GeometryMaskCanvasResource.h
@@ -2,6 +2,7 @@
 
 #pragma once
 
+#include "GeometryMaskConstants.h"
 #include "GeometryMaskTypes.h"
 #include "UObject/Object.h"
 
@@ -28,32 +29,38 @@ class GEOMETRYMASK_API UGeometryMaskCanvasResource
 	GENERATED_BODY()
 
 public:
-	/** (R, G, B, A) */
-	static constexpr int32 MaxNumChannels = 4;
-	static constexpr int32 MaxNumAvailableChannels = 3; // Skip alpha, as it's not universally supported (RT format)
-	static constexpr int32 MaxTextureSize = 8192;
-	
+	static constexpr int32 MaxNumChannels = UE::GeometryMask::MaxNumChannels;
+	static constexpr int32 MaxNumAvailableChannels = UE::GeometryMask::MaxNumChannels;
+	static constexpr int32 MaxTextureSize = UE::GeometryMask::MaxTextureSize;
+
 	UGeometryMaskCanvasResource();
 
 	virtual void BeginDestroy() override;
 
 	/** Will return the first available color channel without a canvas assigned. EGeometryMaskColorChannel::None if not available. */
+	UE_DEPRECATED(5.8, "Color channel management has been deprecated as it's no longer needed")
 	const EGeometryMaskColorChannel GetNextAvailableColorChannel() const;
 
 	/** Requests usage of the given color channel for this resource. Will return true if successful. */
+	UE_DEPRECATED(5.8, "Color channel management has been deprecated as it's no longer needed")
 	bool Checkout(const EGeometryMaskColorChannel InColorChannel, const FGeometryMaskCanvasId& InRequestingCanvasId);
 
 	/** Returns/frees the color channel associated with the given canvas name. Returns true if canvas name found. */
+	UE_DEPRECATED(5.8, "Color channel management has been deprecated as it's no longer needed")
 	bool Checkin(const FGeometryMaskCanvasId& InRequestingCanvasId);
 
 	/** Re-arranges used channels such that they are used sequentially. Returns number of unused channels. */
+	UE_DEPRECATED(5.8, "Color channel management has been deprecated as it's no longer needed")
 	int32 Compact();
 
+	UE_DEPRECATED(5.8, "Color channel management has been deprecated as it's no longer needed")
 	int32 GetNumChannelsUsed() const;
 
+	UE_DEPRECATED(5.8, "Color channel management has been deprecated as it's no longer needed")
 	bool IsAnyChannelUsed() const;
 
 	/** Get the list of dependent canvases, for debug purposes. */
+	UE_DEPRECATED(5.8, "Color channel management has been deprecated as it's no longer needed")
 	TArray<FGeometryMaskCanvasId> GetDependentCanvasIds() const;
 
 	void UpdateViewportSize();
@@ -83,17 +90,22 @@ private:
 
 	FGeometryMaskDrawingContext* GetDrawingContextForLevel(const ULevel* InLevel, uint8 InSceneViewIndex);
 	FGeometryMaskDrawingContext* GetDrawingContextForCanvas(const FGeometryMaskCanvasId& InCanvasId);
+
+	UE_DEPRECATED(5.8, "Color channel management has been deprecated as it's no longer needed")
 	FGeometryMaskDrawingContext* GetDrawingContextForChannel(EGeometryMaskColorChannel InColorChannel);
 
 	[[maybe_unused]] bool RemoveInvalidDrawingContexts();
 
-	UPROPERTY(NoClear, EditFixedSize, meta = (EditFixedOrder))
-	TMap<EGeometryMaskColorChannel, FGeometryMaskCanvasId> DependentCanvasIds;
+	UE_DEPRECATED(5.8, "Color channel management has been deprecated as it's no longer needed")
+	UPROPERTY()
+	TMap<EGeometryMaskColorChannel, FGeometryMaskCanvasId> DependentCanvasIds_DEPRECATED;
 
-	TBitArray<TInlineAllocator<MaxNumChannels>> UsedChannelMask;
+	UE_DEPRECATED(5.8, "Color channel management has been deprecated as it's no longer needed")
+	TBitArray<TInlineAllocator<UE::GeometryMask::MaxNumChannels>> UsedChannelMask;
 
+	UE_DEPRECATED(5.8, "Color channel management has been deprecated as it's no longer needed")
 	UPROPERTY()
-	int32 NumUsedChannels = 0;
+	int32 NumUsedChannels_DEPRECATED = 0;
 
 	/** Canvas used for drawing. Using FCanvas over UCanvas for perf */
 	TSharedPtr<FCanvas> Canvas;
diff --git a/Engine/Plugins/VirtualProduction/GeometryMask/Source/GeometryMask/Public/GeometryMaskConstants.h b/Engine/Plugins/VirtualProduction/GeometryMask/Source/GeometryMask/Public/GeometryMaskConstants.h
new file mode 100644
--- /dev/null
+++ b/Engine/Plugins/VirtualProduction/GeometryMask/Source/GeometryMask/Public/GeometryMaskConstants.h
@@ -0,0 +1,11 @@
+// Copyright Epic Games, Inc. All Rights Reserved.
+
+#pragma once
+
+#include "HAL/Platform.h"
+
+namespace UE::GeometryMask
+{
+	constexpr int32 MaxNumChannels = 1;
+	constexpr int32 MaxTextureSize = 8192;
+}
diff --git a/Engine/Plugins/VirtualProduction/GeometryMask/Source/GeometryMask/Public/GeometryMaskSubsystem.h b/Engine/Plugins/VirtualProduction/GeometryMask/Source/GeometryMask/Public/GeometryMaskSubsystem.h
--- a/Engine/Plugins/VirtualProduction/GeometryMask/Source/GeometryMask/Public/GeometryMaskSubsystem.h
+++ b/Engine/Plugins/VirtualProduction/GeometryMask/Source/GeometryMask/Public/GeometryMaskSubsystem.h
@@ -11,55 +11,73 @@
 
 #include "GeometryMaskSubsystem.generated.h"
 
+#define UE_API GEOMETRYMASK_API
+
 class UGeometryMaskWorldSubsystem;
 
 using FOnGeometryMaskResourceCreated = TMulticastDelegate<void(const UGeometryMaskCanvasResource*)>;
 using FOnGeometryMaskResourceDestroyed = TMulticastDelegate<void(const UGeometryMaskCanvasResource*)>;
 
 /** Maintains the registered named canvases. */
-UCLASS(BlueprintType)
-class GEOMETRYMASK_API UGeometryMaskSubsystem
-	: public UEngineSubsystem
+UCLASS(BlueprintType, MinimalAPI)
+class UGeometryMaskSubsystem : public UEngineSubsystem
 {
 	GENERATED_BODY()
 
 public:
 	/** Returns the default, blank canvas. */
 	UFUNCTION(BlueprintCallable, Category = "Canvas")
-	UGeometryMaskCanvas* GetDefaultCanvas();
+	UE_API UGeometryMaskCanvas* GetDefaultCanvas();
 
-	int32 GetNumCanvasResources() const;
+	UE_DEPRECATED(5.8, "CanvasResources is deprecated. Subsystem no longer manages resources")
+	UE_API int32 GetNumCanvasResources() const;
 
-	const TSet<TObjectPtr<UGeometryMaskCanvasResource>>& GetCanvasResources() const;
+	UE_DEPRECATED(5.8, "CanvasResources is deprecated. Subsystem no longer manages resources")
+	UE_API const TSet<TObjectPtr<UGeometryMaskCanvasResource>>& GetCanvasResources() const;
 
-	void Update(UWorld* InWorld, FSceneViewFamily& InViewFamily);
+	UE_API void Update(UWorld* InWorld, FSceneViewFamily& InViewFamily);
 
 	/** Toggles if no arg given. */
-	void ToggleUpdate(const TOptional<bool>& bInShouldUpdate = {});
+	UE_API void ToggleUpdate(const TOptional<bool>& bInShouldUpdate = {});
 
 	/** Called when a new canvas resource is created. */
-	FOnGeometryMaskResourceCreated& OnGeometryMaskResourceCreated() { return OnGeometryMaskResourceCreatedDelegate; }
+	UE_DEPRECATED(5.8, "OnGeometryMaskResourceCreated is deprecated. Subsystem no longer manages resources")
+	FOnGeometryMaskResourceCreated& OnGeometryMaskResourceCreated()
+	{
+PRAGMA_DISABLE_DEPRECATION_WARNINGS
+		return OnGeometryMaskResourceCreatedDelegate;
+PRAGMA_ENABLE_DEPRECATION_WARNINGS
+	}
 
 	/** Called when a canvas resource is destroyed. */
-	FOnGeometryMaskResourceDestroyed& OnGeometryMaskResourceDestroyed() { return OnGeometryMaskResourceDestroyedDelegate; }
+	UE_DEPRECATED(5.8, "OnGeometryMaskResourceDestroyed is deprecated. Subsystem no longer manages resources")
+	FOnGeometryMaskResourceDestroyed& OnGeometryMaskResourceDestroyed()
+	{
+PRAGMA_DISABLE_DEPRECATION_WARNINGS
+		return OnGeometryMaskResourceDestroyedDelegate;
+PRAGMA_ENABLE_DEPRECATION_WARNINGS
+	}
 
 private:
-	void UpdateLevel(const ULevel* InLevel, UGeometryMaskWorldSubsystem* InWorldSubsystem, FSceneViewFamily& InViewFamily);
+	UE_API void UpdateLevel(const ULevel* InLevel, UGeometryMaskWorldSubsystem* InWorldSubsystem, FSceneViewFamily& InViewFamily);
 
 	/** Find and assign the next available resource to the given canvas. */
-	void AssignResourceToCanvas(UGeometryMaskCanvas* InCanvas);
+	UE_DEPRECATED(5.8, "AssignResourceToCanvas is deprecated. Subsystem no longer manages resources")
+	UE_API void AssignResourceToCanvas(UGeometryMaskCanvas* InCanvas);
 
-	/** Re-arrange used canvases such that they use as few Canvas Resource's as possible.
-	 *  Note that this could cause momentary visual artifacts.
-	 *  This won't check for unused canvas's, only resource channels. */
-	void CompactResources();
+	UE_DEPRECATED(5.8, "CompactResources is deprecated. Subsystem no longer manages resources")
+	UE_API void CompactResources();
 
-	void OnWorldDestroyed(UWorld* InWorld);
+	UE_DEPRECATED(5.8, "OnWorldDestroyed is deprecated. Subsystem no longer manages resources")
+	UE_API void OnWorldDestroyed(UWorld* InWorld);
 
 private:
 	friend class UGeometryMaskWorldSubsystem;
 
+	UE_DEPRECATED(5.8, "OnGeometryMaskResourceCreatedDelegate is deprecated. Subsystem no longer manages resources")
 	FOnGeometryMaskResourceCreated OnGeometryMaskResourceCreatedDelegate;
+
+	UE_DEPRECATED(5.8, "OnGeometryMaskResourceDestroyedDelegate is deprecated. Subsystem no longer manages resources")
 	FOnGeometryMaskResourceDestroyed OnGeometryMaskResourceDestroyedDelegate;
 
 	std::atomic<bool> bDoUpdates = true;
@@ -68,6 +86,9 @@ private:
 	TObjectPtr<UGeometryMaskCanvas> DefaultCanvas;
 
 	/** Pool of GPU/Texture resources used by the canvases. */
-	UPROPERTY(Getter)
-	TSet<TObjectPtr<UGeometryMaskCanvasResource>> CanvasResources;
+	UE_DEPRECATED(5.8, "CanvasResources is deprecated. Subsystem no longer manages resources")
+	UPROPERTY()
+	TSet<TObjectPtr<UGeometryMaskCanvasResource>> CanvasResources_DEPRECATED;
 };
+
+#undef UE_API
diff --git a/Engine/Plugins/VirtualProduction/GeometryMask/Source/GeometryMask/Public/GeometryMaskTypes.h b/Engine/Plugins/VirtualProduction/GeometryMask/Source/GeometryMask/Public/GeometryMaskTypes.h
--- a/Engine/Plugins/VirtualProduction/GeometryMask/Source/GeometryMask/Public/GeometryMaskTypes.h
+++ b/Engine/Plugins/VirtualProduction/GeometryMask/Source/GeometryMask/Public/GeometryMaskTypes.h
@@ -26,7 +26,9 @@ ENUM_CLASS_FLAGS(EGeometryMaskColorChannel)
 
 namespace UE::GeometryMask
 {
+PRAGMA_DISABLE_DEPRECATION_WARNINGS
 	/** Use to convert between a multiplied vector and an EGeometryMaskColorChannel enum value. */
+	UE_DEPRECATED(5.8, "Color channel management has been deprecated as it's no longer needed")
 	static TMap<FLinearColor, EGeometryMaskColorChannel> VectorToMaskChannelEnum =
 	{
 		{ {1, 0, 0, 0}, EGeometryMaskColorChannel::Red },
@@ -34,8 +36,9 @@ namespace UE::GeometryMask
 		{ {0, 0, 1, 0}, EGeometryMaskColorChannel::Blue },
 		{ {0, 0, 0, 1}, EGeometryMaskColorChannel::Alpha }
 	};
-	
+
 	/** Use to convert between a EGeometryMaskColorChannel enum value and a multiplied vector. */
+	UE_DEPRECATED(5.8, "Color channel management has been deprecated as it's no longer needed")
 	static TMap<EGeometryMaskColorChannel, FLinearColor> MaskChannelEnumToVector =
 	{
 		{ EGeometryMaskColorChannel::Red, {1, 0, 0, 0} },
@@ -45,6 +48,7 @@ namespace UE::GeometryMask
 	};
 
 	/** Use to convert between a EGeometryMaskColorChannel enum value and an FColor. */
+	UE_DEPRECATED(5.8, "Color channel management has been deprecated as it's no longer needed")
 	static TMap<EGeometryMaskColorChannel, FColor> MaskChannelEnumToColor =
 	{
 		{ EGeometryMaskColorChannel::Red, {255, 0, 0, 0} },
@@ -52,6 +56,7 @@ namespace UE::GeometryMask
 		{ EGeometryMaskColorChannel::Blue, {0, 0, 255, 0} },
 		{ EGeometryMaskColorChannel::Alpha, {0, 0, 0, 255} }
 	};	
+PRAGMA_ENABLE_DEPRECATION_WARNINGS
 
 	/** Use to convert between a EGeometryMaskColorChannel and it's value name (as string). */
 	static TMap<EGeometryMaskColorChannel, FString> MaskChannelEnumToString =
@@ -64,9 +69,11 @@ namespace UE::GeometryMask
 		{ EGeometryMaskColorChannel::Num, TEXT("Num/Max") }
 	};
 
+	UE_DEPRECATED(5.8, "Color channel management has been deprecated as it's no longer needed")
 	GEOMETRYMASK_API EGeometryMaskColorChannel VectorToMaskChannel(const FLinearColor& InVector);
 
 	/** Returns a valid ColorChannel (R,G,B or A). */
+	UE_DEPRECATED(5.8, "Color channel management has been deprecated as it's no longer needed")
 	GEOMETRYMASK_API EGeometryMaskColorChannel GetValidMaskChannel(EGeometryMaskColorChannel InColorChannel, bool bInIncludeAlpha = false);
 
 	GEOMETRYMASK_API FStringView ChannelToString(EGeometryMaskColorChannel InColorChannel);
@@ -192,13 +199,21 @@ struct FGeometryMaskWriteParameters
 {
 	GENERATED_BODY()
 
-public:
+	PRAGMA_DISABLE_DEPRECATION_WARNINGS
+	FGeometryMaskWriteParameters() = default;
+	FGeometryMaskWriteParameters(const FGeometryMaskWriteParameters&) = default;
+	FGeometryMaskWriteParameters(FGeometryMaskWriteParameters&&) = default;
+	FGeometryMaskWriteParameters& operator=(const FGeometryMaskWriteParameters&) = default;
+	FGeometryMaskWriteParameters& operator=(FGeometryMaskWriteParameters&&) = default;
+	PRAGMA_ENABLE_DEPRECATION_WARNINGS
+
 	/** Specifies the GeometryMaskCanvas to reference. */
 	UPROPERTY(EditAnywhere, BlueprintReadWrite, Category = "Canvas")
 	FName CanvasName;
-	
-	UPROPERTY(EditAnywhere, Category = "Write")
-	EGeometryMaskColorChannel ColorChannel = EGeometryMaskColorChannel::Red;
+
+	UE_DEPRECATED(5.8, "Color channel management has been deprecated as it's no longer needed")
+	UPROPERTY()
+	EGeometryMaskColorChannel ColorChannel_DEPRECATED = EGeometryMaskColorChannel::Red;
 
 	UPROPERTY(EditAnywhere, Category = "Write")
 	EGeometryMaskCompositeOperation OperationType = EGeometryMaskCompositeOperation::Add;
diff --git a/Engine/Plugins/VirtualProduction/GeometryMask/Source/GeometryMask/Public/GeometryMaskWorldSubsystem.h b/Engine/Plugins/VirtualProduction/GeometryMask/Source/GeometryMask/Public/GeometryMaskWorldSubsystem.h
--- a/Engine/Plugins/VirtualProduction/GeometryMask/Source/GeometryMask/Public/GeometryMaskWorldSubsystem.h
+++ b/Engine/Plugins/VirtualProduction/GeometryMask/Source/GeometryMask/Public/GeometryMaskWorldSubsystem.h
@@ -34,6 +34,14 @@ public:
 	/** Retrieves a Canvas, uniquely identified by this world and the canvas name. */
 	UGeometryMaskCanvas* GetNamedCanvas(const ULevel* InLevel, FName InName);
 
+	struct FCanvasContext
+	{
+		FName CanvasName;
+		TWeakObjectPtr<const ULevel> LevelWeak;
+	};
+	/** Invokes the given function ref for each existing canvas */
+	void ForEachCanvas(TFunctionRef<void(UGeometryMaskCanvas*, const FCanvasContext&)> InFunc) const;
+
 	/** Returns all registered canvas names for this world. */
 	TArray<FName> GetCanvasNames(const ULevel* InLevel);
 	
@@ -57,10 +65,6 @@ private:
 
 	FGeometryMaskLevelState& FindOrAddLevelState(const ULevel* InLevel);
 
-	void OnCanvasActivated(UGeometryMaskCanvas* InCanvas);
-	void OnCanvasDeactivated(UGeometryMaskCanvas* InCanvas);
-
-private:
 	friend class UGeometryMaskSubsystem;
 
 	TSharedPtr<FGeometryMaskSceneViewExtension> GeometryMaskSceneViewExtension;
diff --git a/Engine/Plugins/VirtualProduction/GeometryMask/Source/GeometryMaskEditor/Private/GeometryMaskEditorModule.cpp b/Engine/Plugins/VirtualProduction/GeometryMask/Source/GeometryMaskEditor/Private/GeometryMaskEditorModule.cpp
--- a/Engine/Plugins/VirtualProduction/GeometryMask/Source/GeometryMaskEditor/Private/GeometryMaskEditorModule.cpp
+++ b/Engine/Plugins/VirtualProduction/GeometryMask/Source/GeometryMaskEditor/Private/GeometryMaskEditorModule.cpp
@@ -10,10 +10,8 @@
 #include "Materials/Material.h"
 #include "Styling/SlateIconFinder.h"
 #include "ViewModels/GMECanvasListViewModel.h"
-#include "ViewModels/GMEResourceListViewModel.h"
 #include "Widgets/Docking/SDockTab.h"
 #include "Widgets/SGMECanvasList.h"
-#include "Widgets/SGMEResourceList.h"
 #include "WorkspaceMenuStructure.h"
 #include "WorkspaceMenuStructureModule.h"
 
@@ -32,25 +30,12 @@ void FGeometryMaskEditorModule::StartupModule()
 				[this](const FSpawnTabArgs&)
 				{
 					TSharedRef<FGMECanvasListViewModel> CanvasListViewModel = FGMECanvasListViewModel::Create();
-					TSharedRef<FGMEResourceListViewModel> ResourceListViewModel = FGMEResourceListViewModel::Create();
-
 					TSharedRef<SDockTab> DockTab =
 						SNew(SDockTab)
 						.TabRole(ETabRole::NomadTab)
 						.Content()
 						[
-							SNew(SHorizontalBox)
-							+ SHorizontalBox::Slot()
-							.FillWidth(0.5)
-							[
-								SNew(SGMECanvasList, CanvasListViewModel)
-							]
-							
-							+ SHorizontalBox::Slot()
-							.FillWidth(0.5)
-							[
-								SNew(SGMEResourceList, ResourceListViewModel)
-							]
+							SNew(SGMECanvasList, CanvasListViewModel)
 						];
 
 					return DockTab;
diff --git a/Engine/Plugins/VirtualProduction/GeometryMask/Source/GeometryMaskEditor/Private/ViewModels/GMECanvasItemViewModel.cpp b/Engine/Plugins/VirtualProduction/GeometryMask/Source/GeometryMaskEditor/Private/ViewModels/GMECanvasItemViewModel.cpp
--- a/Engine/Plugins/VirtualProduction/GeometryMask/Source/GeometryMaskEditor/Private/ViewModels/GMECanvasItemViewModel.cpp
+++ b/Engine/Plugins/VirtualProduction/GeometryMask/Source/GeometryMaskEditor/Private/ViewModels/GMECanvasItemViewModel.cpp
@@ -24,7 +24,6 @@ FGMECanvasItemViewModel::FGMECanvasItemViewModel(
 	if (const UGeometryMaskCanvas* Canvas = InCanvas.Get())
 	{
 		CanvasId = Canvas->GetCanvasId();		
-		ColorChannel = Canvas->GetColorChannel();
 		CanvasTextureWeak = Canvas->GetTexture();
 		KnownWriterCount = Canvas->GetWriters().Num();
 
diff --git a/Engine/Plugins/VirtualProduction/GeometryMask/Source/GeometryMaskEditor/Private/ViewModels/GMECanvasItemViewModel.h b/Engine/Plugins/VirtualProduction/GeometryMask/Source/GeometryMaskEditor/Private/ViewModels/GMECanvasItemViewModel.h
--- a/Engine/Plugins/VirtualProduction/GeometryMask/Source/GeometryMaskEditor/Private/ViewModels/GMECanvasItemViewModel.h
+++ b/Engine/Plugins/VirtualProduction/GeometryMask/Source/GeometryMaskEditor/Private/ViewModels/GMECanvasItemViewModel.h
@@ -38,7 +38,6 @@ public:
 	const FGeometryMaskCanvasId& GetCanvasId() const;
 	FName GetCanvasName() const;
 	const FText& GetCanvasInfo() const;
-	const EGeometryMaskColorChannel GetColorChannel() const { return ColorChannel; }
 	const UTexture* GetCanvasTexture() const;
 	float GetMemoryUsage();
 
@@ -48,7 +47,6 @@ private:
 private:
 	FGeometryMaskCanvasId CanvasId;
 	FText InfoText;
-	EGeometryMaskColorChannel ColorChannel;
 	TWeakObjectPtr<const UGeometryMaskCanvas> CanvasWeak;
 	TWeakObjectPtr<UTexture> CanvasTextureWeak;
 	int32 KnownReaderCount;
diff --git a/Engine/Plugins/VirtualProduction/GeometryMask/Source/GeometryMaskEditor/Private/ViewModels/GMEResourceItemViewModel.cpp b/Engine/Plugins/VirtualProduction/GeometryMask/Source/GeometryMaskEditor/Private/ViewModels/GMEResourceItemViewModel.cpp
deleted file mode 100644
diff --git a/Engine/Plugins/VirtualProduction/GeometryMask/Source/GeometryMaskEditor/Private/ViewModels/GMEResourceItemViewModel.h b/Engine/Plugins/VirtualProduction/GeometryMask/Source/GeometryMaskEditor/Private/ViewModels/GMEResourceItemViewModel.h
deleted file mode 100644
diff --git a/Engine/Plugins/VirtualProduction/GeometryMask/Source/GeometryMaskEditor/Private/ViewModels/GMEResourceListViewModel.cpp b/Engine/Plugins/VirtualProduction/GeometryMask/Source/GeometryMaskEditor/Private/ViewModels/GMEResourceListViewModel.cpp
deleted file mode 100644
diff --git a/Engine/Plugins/VirtualProduction/GeometryMask/Source/GeometryMaskEditor/Private/ViewModels/GMEResourceListViewModel.h b/Engine/Plugins/VirtualProduction/GeometryMask/Source/GeometryMaskEditor/Private/ViewModels/GMEResourceListViewModel.h
deleted file mode 100644
diff --git a/Engine/Plugins/VirtualProduction/GeometryMask/Source/GeometryMaskEditor/Private/Widgets/SGMECanvasItem.cpp b/Engine/Plugins/VirtualProduction/GeometryMask/Source/GeometryMaskEditor/Private/Widgets/SGMECanvasItem.cpp
--- a/Engine/Plugins/VirtualProduction/GeometryMask/Source/GeometryMaskEditor/Private/Widgets/SGMECanvasItem.cpp
+++ b/Engine/Plugins/VirtualProduction/GeometryMask/Source/GeometryMaskEditor/Private/Widgets/SGMECanvasItem.cpp
@@ -20,7 +20,6 @@ void SGMECanvasItem::Construct(
 					.Invert(true)
 					.SolidBackground(true)
 					.CanvasId(ViewModel->GetCanvasId())
-					.Channel(ViewModel->GetColorChannel())
 					.Opacity(1.0f);
 
 	SGMEImageItem::Construct(
diff --git a/Engine/Plugins/VirtualProduction/GeometryMask/Source/GeometryMaskEditor/Private/Widgets/SGMEResourceItem.cpp b/Engine/Plugins/VirtualProduction/GeometryMask/Source/GeometryMaskEditor/Private/Widgets/SGMEResourceItem.cpp
deleted file mode 100644
diff --git a/Engine/Plugins/VirtualProduction/GeometryMask/Source/GeometryMaskEditor/Private/Widgets/SGMEResourceItem.h b/Engine/Plugins/VirtualProduction/GeometryMask/Source/GeometryMaskEditor/Private/Widgets/SGMEResourceItem.h
deleted file mode 100644
diff --git a/Engine/Plugins/VirtualProduction/GeometryMask/Source/GeometryMaskEditor/Private/Widgets/SGMEResourceList.cpp b/Engine/Plugins/VirtualProduction/GeometryMask/Source/GeometryMaskEditor/Private/Widgets/SGMEResourceList.cpp
deleted file mode 100644
diff --git a/Engine/Plugins/VirtualProduction/GeometryMask/Source/GeometryMaskEditor/Private/Widgets/SGMEResourceList.h b/Engine/Plugins/VirtualProduction/GeometryMask/Source/GeometryMaskEditor/Private/Widgets/SGMEResourceList.h
deleted file mode 100644
diff --git a/Engine/Plugins/VirtualProduction/GeometryMask/Source/GeometryMaskEditor/Private/Widgets/SGeometryMaskCanvasPreview.cpp b/Engine/Plugins/VirtualProduction/GeometryMask/Source/GeometryMaskEditor/Private/Widgets/SGeometryMaskCanvasPreview.cpp
--- a/Engine/Plugins/VirtualProduction/GeometryMask/Source/GeometryMaskEditor/Private/Widgets/SGeometryMaskCanvasPreview.cpp
+++ b/Engine/Plugins/VirtualProduction/GeometryMask/Source/GeometryMaskEditor/Private/Widgets/SGeometryMaskCanvasPreview.cpp
@@ -27,7 +27,6 @@ namespace UE::GeometryMaskEditor::Private
 	constexpr const TCHAR* DefaultTexturePath = TEXT("/Engine/EngineResources/Black");
 
 	static FName TextureParameterName = TEXT("Mask_Texture");
-	static FName ChannelParameterName = TEXT("Mask_Channel");
 	static FName InvertParameterName = TEXT("Mask_Invert");
 	static FName PaddingParameterName = TEXT("Mask_Padding");
 	static FName OpacityParameterName = TEXT("OpacityMultiplier");
@@ -43,12 +42,16 @@ SGeometryMaskCanvasPreview::SGeometryMaskCanvasPreview()
 SGeometryMaskCanvasPreview::~SGeometryMaskCanvasPreview()
 {
 	FEditorSupportDelegates::PrepareToCleanseEditorObject.RemoveAll(this);
+
+	FWorldDelegates::OnWorldCleanup.Remove(OnWorldCleanupDelegateHandle);
+	OnWorldCleanupDelegateHandle.Reset();
 }
 
 void SGeometryMaskCanvasPreview::Construct(const FArguments& InArgs)
 {
+	OnWorldCleanupDelegateHandle = FWorldDelegates::OnWorldCleanup.AddSP(this, &SGeometryMaskCanvasPreview::OnWorldCleanup);
+
 	CanvasId = InArgs._CanvasId;
-	ColorChannel = InArgs._Channel;
 	bShowPaddingFrame = InArgs._PaddingFrameVisibility;
 	Invert = InArgs._Invert;
 	HasSolidBackground = InArgs._SolidBackground;
@@ -136,28 +139,11 @@ const FName SGeometryMaskCanvasPreview::GetCanvasName() const
 
 const EGeometryMaskColorChannel SGeometryMaskCanvasPreview::GetColorChannel(const bool bOnlyValid) const
 {
-	EGeometryMaskColorChannel Result = ColorChannel.Get(EGeometryMaskColorChannel::Red);
-	if (bOnlyValid)
-	{
-		Result = UE::GeometryMask::GetValidMaskChannel(Result, false);
-	}
-
-	return Result;
+	return EGeometryMaskColorChannel::Red;
 }
 
 void SGeometryMaskCanvasPreview::SetColorChannel(const EGeometryMaskColorChannel InColorChannel)
 {
-	if (InColorChannel == EGeometryMaskColorChannel::None)
-	{
-		return;
-	}
-
-	ColorChannel.Set(InColorChannel);
-	
-	if (PreviewMID)
-	{
-		PreviewMID->SetVectorParameterValue(UE::GeometryMaskEditor::Private::ChannelParameterName, UE::GeometryMask::MaskChannelEnumToVector[InColorChannel]);
-	}
 }
 
 const bool SGeometryMaskCanvasPreview::IsPaddingFrameVisible() const
@@ -304,7 +290,6 @@ void SGeometryMaskCanvasPreview::UpdateBrush(const UGeometryMaskCanvas* InCanvas
 	if (bHasValidMID)
 	{
 		PreviewMID->SetTextureParameterValue(UE::GeometryMaskEditor::Private::TextureParameterName, InTexture);
-		PreviewMID->SetVectorParameterValue(UE::GeometryMaskEditor::Private::ChannelParameterName, UE::GeometryMask::MaskChannelEnumToVector[GetColorChannel()]);
 		PreviewMID->SetScalarParameterValue(UE::GeometryMaskEditor::Private::InvertParameterName, IsInverted() ? 1.0f : 0.0f);
 	}
 
@@ -321,3 +306,12 @@ void SGeometryMaskCanvasPreview::UpdateBrush(const UGeometryMaskCanvas* InCanvas
 		AspectRatio.Set(ImageSize.X / ImageSize.Y);
 	}
 }
+
+void SGeometryMaskCanvasPreview::OnWorldCleanup(UWorld* InWorld, bool bInSessionEnded, bool bInCleanupResources)
+{
+	const ULevel* const CanvasLevel = CanvasId.Get().Level.ResolveObjectPtr();
+	if (!CanvasLevel || CanvasLevel->GetWorld() == InWorld || CanvasLevel->GetOuter() == InWorld)
+	{
+		UpdateBrush(/*Canvas*/nullptr, /*Texture*/nullptr);
+	}
+}
diff --git a/Engine/Plugins/VirtualProduction/GeometryMask/Source/GeometryMaskEditor/Public/Widgets/SGeometryMaskCanvasPreview.h b/Engine/Plugins/VirtualProduction/GeometryMask/Source/GeometryMaskEditor/Public/Widgets/SGeometryMaskCanvasPreview.h
--- a/Engine/Plugins/VirtualProduction/GeometryMask/Source/GeometryMaskEditor/Public/Widgets/SGeometryMaskCanvasPreview.h
+++ b/Engine/Plugins/VirtualProduction/GeometryMask/Source/GeometryMaskEditor/Public/Widgets/SGeometryMaskCanvasPreview.h
@@ -19,7 +19,7 @@ class GEOMETRYMASKEDITOR_API SGeometryMaskCanvasPreview
 public:
 	SLATE_BEGIN_ARGS(SGeometryMaskCanvasPreview) {}
 		SLATE_ATTRIBUTE(FGeometryMaskCanvasId, CanvasId)
-		SLATE_ATTRIBUTE(EGeometryMaskColorChannel, Channel)
+		SLATE_ATTRIBUTE_DEPRECATED(EGeometryMaskColorChannel, Channel, 5.8, "Color channel management has been deprecated as it's no longer needed")
 		SLATE_ATTRIBUTE(bool, PaddingFrameVisibility)
 		SLATE_ATTRIBUTE(bool, Invert)
 		SLATE_ATTRIBUTE(bool, SolidBackground)
@@ -43,10 +43,10 @@ public:
 	/** Get the name of the currently referenced canvas. */
 	const FName GetCanvasName() const;
 
-	/** Get the color channel to display for the canvas.  */
+	UE_DEPRECATED(5.8, "Color channel management has been deprecated as it's no longer needed")
 	const EGeometryMaskColorChannel GetColorChannel(const bool bOnlyValid = true) const;
 
-	/** Set the color channel to display for the canvas.  */
+	UE_DEPRECATED(5.8, "Color channel management has been deprecated as it's no longer needed")
 	void SetColorChannel(const EGeometryMaskColorChannel InColorChannel);
 
 	const bool IsPaddingFrameVisible() const;
@@ -87,6 +87,12 @@ private:
 
 	void UpdateBrush(const UGeometryMaskCanvas* InCanvas, UTexture* InTexture);
 
+	/**
+	 * Called when the world is being cleaned up
+	 * Used to remove the mask texture reference from the preview MIDs if the textures belong to the world being cleaned up
+	 */
+	void OnWorldCleanup(UWorld* InWorld, bool bInSessionEnded, bool bInCleanupResources);
+
 private:
 	TSharedPtr<FSlateBrush> PreviewBrush;
 	FSoftObjectPath PreviewMaterialPath;
@@ -100,9 +106,6 @@ private:
 	/** Id of the currently referenced canvas. */
 	TAttribute<FGeometryMaskCanvasId> CanvasId;
 
-	/** Color channel to display for the canvas.  */
-	TAttribute<EGeometryMaskColorChannel> ColorChannel;
-
 	TAttribute<bool> bShowPaddingFrame;
 
 	/** Whether to invert the display of the canvas.  */
@@ -116,4 +119,7 @@ private:
 
 	/** Stores the aspect ratio of the referenced canvas. */
 	TAttribute<FOptionalSize> AspectRatio;
+
+	/** Handle to the delegate called when world gets cleaned up */
+	FDelegateHandle OnWorldCleanupDelegateHandle;
 };
