Geometry Mask: added trace event scopes when drawing mask

#Synced-CL 49036542

diff --git a/Engine/Plugins/VirtualProduction/GeometryMask/Source/GeometryMask/Private/GeometryMaskCanvasResource.cpp b/Engine/Plugins/VirtualProduction/GeometryMask/Source/GeometryMask/Private/GeometryMaskCanvasResource.cpp
--- a/Engine/Plugins/VirtualProduction/GeometryMask/Source/GeometryMask/Private/GeometryMaskCanvasResource.cpp
+++ b/Engine/Plugins/VirtualProduction/GeometryMask/Source/GeometryMask/Private/GeometryMaskCanvasResource.cpp
@@ -319,6 +319,8 @@ void UGeometryMaskCanvasResource::Draw(const ULevel* InLevel, FSceneView& InView
 		return;
 	}
 
+	TRACE_CPUPROFILER_EVENT_SCOPE(UGeometryMaskCanvasResource::Draw);
+
 	UE_LOG(LogGeometryMask, VeryVerbose, TEXT("UGeometryMaskCanvasResource::Update Level: %s"), *InLevel->GetName());
 
 	if (UCanvasRenderTarget2D* Texture = GetRenderTargetTexture())
@@ -367,17 +369,19 @@ void UGeometryMaskCanvasResource::Draw(const ULevel* InLevel, FSceneView& InView
 			Canvas->SetBaseTransform(DrawingContextRef.ViewProjectionMatrix);
 			{
 				TRACE_CPUPROFILER_EVENT_SCOPE(UGeometryMaskCanvasResource::Draw::OnDrawToCanvas);
-				
 				OnDrawToCanvas().Broadcast(DrawingContextRef, Canvas.Get());
 			}
 
 			// Restore original transform
 			Canvas->SetBaseTransform(CanvasMatrix);
 		}
-		
+
 		// End
 		{
-			Canvas->Flush_GameThread();
+			{
+				TRACE_CPUPROFILER_EVENT_SCOPE(UGeometryMaskCanvasResource::Draw::Flush);
+				Canvas->Flush_GameThread();
+			}
 
 			// Post Process
 			if (bApplyDF || bApplyBlur)
diff --git a/Engine/Plugins/VirtualProduction/GeometryMask/Source/GeometryMask/Private/GeometryMaskSubsystem.cpp b/Engine/Plugins/VirtualProduction/GeometryMask/Source/GeometryMask/Private/GeometryMaskSubsystem.cpp
--- a/Engine/Plugins/VirtualProduction/GeometryMask/Source/GeometryMask/Private/GeometryMaskSubsystem.cpp
+++ b/Engine/Plugins/VirtualProduction/GeometryMask/Source/GeometryMask/Private/GeometryMaskSubsystem.cpp
@@ -78,6 +78,8 @@ void UGeometryMaskSubsystem::UpdateLevel(const ULevel* InLevel, UGeometryMaskWor
 		return;
 	}
 
+	TRACE_CPUPROFILER_EVENT_SCOPE(UGeometryMaskSubsystem::UpdateLevel);
+
 	int32 ViewIndex = 0;
 
 	for (const FSceneView*& View : InViewFamily.Views)
