Motion Design: added support for nested masking

#virtualized
#Synced-CL 49298335

diff --git a/Engine/Plugins/VirtualProduction/Avalanche/Source/AvalancheMask/Private/AvaMaskTypes.cpp b/Engine/Plugins/VirtualProduction/Avalanche/Source/AvalancheMask/Private/AvaMaskTypes.cpp
--- a/Engine/Plugins/VirtualProduction/Avalanche/Source/AvalancheMask/Private/AvaMaskTypes.cpp
+++ b/Engine/Plugins/VirtualProduction/Avalanche/Source/AvalancheMask/Private/AvaMaskTypes.cpp
@@ -4,22 +4,24 @@
 
 #include "AvaMaskLog.h"
 #include "AvaMaskUtilities.h"
-#include "Engine/Texture.h"
+#include "Engine/TextureRenderTarget2DArray.h"
 #include "Materials/MaterialInstanceDynamic.h"
 #include "Materials/MaterialInterface.h"
+#include "Misc/EnumerateRange.h"
 
 bool FAvaMask2DMaterialParameters::HasSameParameters(const FAvaMask2DMaterialParameters& InOther) const
 {
 	constexpr float Tolerance = UE_SMALL_NUMBER;
 
-	return BlendMode == InOther.BlendMode
-		&& Texture == InOther.Texture
-		&& FMath::IsNearlyEqual(BaseOpacity, InOther.BaseOpacity, Tolerance)
-		&& bInvert == InOther.bInvert
-		&& Padding.Equals(InOther.Padding, Tolerance)
+	return bInvert == InOther.bInvert
 		&& bApplyFeathering == InOther.bApplyFeathering
+		&& BlendMode == InOther.BlendMode
+		&& RenderTarget == InOther.RenderTarget
+		&& MaskIndices.Equals(InOther.MaskIndices, Tolerance)
+		&& FMath::IsNearlyEqual(BaseOpacity, InOther.BaseOpacity, Tolerance)
 		&& FMath::IsNearlyEqual(InnerFeatherRadius, InOther.InnerFeatherRadius, Tolerance)
-		&& FMath::IsNearlyEqual(OuterFeatherRadius, InOther.OuterFeatherRadius, Tolerance);
+		&& FMath::IsNearlyEqual(OuterFeatherRadius, InOther.OuterFeatherRadius, Tolerance)
+		&& Padding.Equals(InOther.Padding, Tolerance);
 }
 
 bool FAvaMask2DMaterialParameters::ApplyToMID(UMaterialInstanceDynamic* InMaterial) const
@@ -33,7 +35,8 @@ bool FAvaMask2DMaterialParameters::ApplyToMID(UMaterialInstanceDynamic* InMateri
 	InMaterial->BasePropertyOverrides.bOverride_BlendMode = true;
 	InMaterial->BasePropertyOverrides.BlendMode = BlendMode;
 
-	InMaterial->SetTextureParameterValueByInfo(UE::AvaMask::Internal::TextureParameterInfo, Texture);
+	InMaterial->SetTextureParameterValueByInfo(UE::AvaMask::Internal::MaskTexturesParameterInfo, RenderTarget);
+	InMaterial->SetVectorParameterValueByInfo(UE::AvaMask::Internal::MaskIndicesParameterInfo, MaskIndices);
 	InMaterial->SetScalarParameterValueByInfo(UE::AvaMask::Internal::BaseOpacityParameterInfo, FMath::Clamp(BaseOpacity, 0.0f, 1.0f));
 	InMaterial->SetScalarParameterValueByInfo(UE::AvaMask::Internal::InvertParameterInfo, bInvert ? 1.f : 0.f);
 	InMaterial->SetVectorParameterValueByInfo(UE::AvaMask::Internal::PaddingParameterInfo, FLinearColor(FMath::Max(0, Padding.X), FMath::Max(0, Padding.Y), 0, 0));
@@ -42,38 +45,14 @@ bool FAvaMask2DMaterialParameters::ApplyToMID(UMaterialInstanceDynamic* InMateri
 }
 
 bool FAvaMask2DMaterialParameters::StoreFromMaterial(const UMaterialInterface* InMaterial)
-{
-	return SetFromMaterial(InMaterial
-		, UE::AvaMask::Internal::TextureParameterInfo
-		, UE::AvaMask::Internal::InvertParameterInfo
-		, UE::AvaMask::Internal::BaseOpacityParameterInfo
-		, UE::AvaMask::Internal::PaddingParameterInfo
-		, UE::AvaMask::Internal::FeatherParameterInfo);
-}
-
-bool FAvaMask2DMaterialParameters::SetFromMaterial(const UMaterialInterface* InMaterial
-	, const FMaterialParameterInfo& InTextureParameter
-	, const FMaterialParameterInfo& InChannelParameter
-	, const FMaterialParameterInfo& InInvertParameter
-	, const FMaterialParameterInfo& InBaseOpacityParameter
-	, const FMaterialParameterInfo& InPaddingParameter
-	, const FMaterialParameterInfo& InFeatherParameter)
-{
-	return SetFromMaterial(InMaterial, InTextureParameter, InInvertParameter, InBaseOpacityParameter, InPaddingParameter, InFeatherParameter);
-}
-
-bool FAvaMask2DMaterialParameters::SetFromMaterial(const UMaterialInterface* InMaterial
-	, const FMaterialParameterInfo& InTextureParameter
-	, const FMaterialParameterInfo& InInvertParameter
-	, const FMaterialParameterInfo& InBaseOpacityParameter
-	, const FMaterialParameterInfo& InPaddingParameter
-	, const FMaterialParameterInfo& InFeatherParameter)
 {
 	if (!ensureAlways(InMaterial))
 	{
 		return false;
 	}
 
+	constexpr float Tolerance = UE_SMALL_NUMBER;
+
 	BlendMode = InMaterial->GetBlendMode();
 
 	auto LogMissingParameter = [InMaterial](const FMaterialParameterInfo& InParameter)
@@ -81,23 +60,34 @@ bool FAvaMask2DMaterialParameters::SetFromMaterial(const UMaterialInterface* InM
 			UE_LOG(LogAvaMask, Warning, TEXT("Could not find parameter '%s' in material '%s'"), *InParameter.ToString(), *InMaterial->GetFullName());;
 		};
 
-	// Texture
+	// Mask Textures
 	{
 		UTexture* FoundTexture = nullptr;
-		if (!InMaterial->GetTextureParameterValue(InTextureParameter, FoundTexture))
+		if (!InMaterial->GetTextureParameterValue(UE::AvaMask::Internal::MaskTexturesParameterInfo, FoundTexture))
+		{
+			LogMissingParameter(UE::AvaMask::Internal::MaskTexturesParameterInfo);
+			return false;
+		}
+		RenderTarget = Cast<UTextureRenderTarget2DArray>(FoundTexture);
+	}
+
+	// Mask Indices
+	{
+		FLinearColor FoundMaskIndices;
+		if (!InMaterial->GetVectorParameterValue(UE::AvaMask::Internal::MaskIndicesParameterInfo, FoundMaskIndices))
 		{
-			LogMissingParameter(InTextureParameter);
+			LogMissingParameter(UE::AvaMask::Internal::MaskIndicesParameterInfo);
 			return false;
 		}
-		Texture = FoundTexture;
+		MaskIndices = FoundMaskIndices;
 	}
 
 	// Base Opacity
 	{
 		float FoundValue = 0.0f;
-		if (!InMaterial->GetScalarParameterValue(InBaseOpacityParameter, FoundValue))
+		if (!InMaterial->GetScalarParameterValue(UE::AvaMask::Internal::BaseOpacityParameterInfo, FoundValue))
 		{
-			LogMissingParameter(InTextureParameter);
+			LogMissingParameter(UE::AvaMask::Internal::BaseOpacityParameterInfo);
 			return false;
 		}
 		BaseOpacity = FoundValue;
@@ -106,20 +96,20 @@ bool FAvaMask2DMaterialParameters::SetFromMaterial(const UMaterialInterface* InM
 	// Invert
 	{
 		float FoundValue = 0.0f;
-		if (!InMaterial->GetScalarParameterValue(InInvertParameter, FoundValue))
+		if (!InMaterial->GetScalarParameterValue(UE::AvaMask::Internal::InvertParameterInfo, FoundValue))
 		{
-			LogMissingParameter(InTextureParameter);
+			LogMissingParameter(UE::AvaMask::Internal::InvertParameterInfo);
 			return false;
 		}
-		bInvert = FoundValue > 0.1;
+		bInvert = !FMath::IsNearlyZero(FoundValue, Tolerance);
 	}
 
 	// Padding
 	{
 		FLinearColor FoundValue;
-		if (!InMaterial->GetVectorParameterValue(InPaddingParameter, FoundValue))
+		if (!InMaterial->GetVectorParameterValue(UE::AvaMask::Internal::PaddingParameterInfo, FoundValue))
 		{
-			LogMissingParameter(InTextureParameter);
+			LogMissingParameter(UE::AvaMask::Internal::PaddingParameterInfo);
 			return false;			
 		}
 		Padding = FVector2f(FoundValue.R, FoundValue.G);
@@ -128,12 +118,12 @@ bool FAvaMask2DMaterialParameters::SetFromMaterial(const UMaterialInterface* InM
 	// Feathering
 	{
 		FLinearColor FoundValue;
-		if (!InMaterial->GetVectorParameterValue(InFeatherParameter, FoundValue))
+		if (!InMaterial->GetVectorParameterValue(UE::AvaMask::Internal::FeatherParameterInfo, FoundValue))
 		{
-			LogMissingParameter(InTextureParameter);
+			LogMissingParameter(UE::AvaMask::Internal::FeatherParameterInfo);
 			return false;
 		}
-		bApplyFeathering = !FMath::IsNearlyZero(FoundValue.R);
+		bApplyFeathering = !FMath::IsNearlyZero(FoundValue.R, Tolerance);
 		OuterFeatherRadius = FoundValue.G;
 		InnerFeatherRadius = FoundValue.B;
 	}
diff --git a/Engine/Plugins/VirtualProduction/Avalanche/Source/AvalancheMask/Private/AvaMaskUtilities.h b/Engine/Plugins/VirtualProduction/Avalanche/Source/AvalancheMask/Private/AvaMaskUtilities.h
--- a/Engine/Plugins/VirtualProduction/Avalanche/Source/AvalancheMask/Private/AvaMaskUtilities.h
+++ b/Engine/Plugins/VirtualProduction/Avalanche/Source/AvalancheMask/Private/AvaMaskUtilities.h
@@ -16,7 +16,8 @@ class UMaterialInstance;
 namespace UE::AvaMask::Internal
 {
 	/** Note that this assumes the parameters have the same name for all materials, which they should if using the provided material function. */
-	static FMaterialParameterInfo TextureParameterInfo = { TEXT("Mask_Texture") };
+	static FMaterialParameterInfo MaskTexturesParameterInfo = { TEXT("Mask_Textures") };
+	static FMaterialParameterInfo MaskIndicesParameterInfo = { TEXT("Mask_TextureIndexVector") };
 	static FMaterialParameterInfo BaseOpacityParameterInfo = { TEXT("Mask_BaseOpacity") };
 	static FMaterialParameterInfo InvertParameterInfo = { TEXT("Mask_Invert") };
 	static FMaterialParameterInfo PaddingParameterInfo = { TEXT("Mask_Padding") };
diff --git a/Engine/Plugins/VirtualProduction/Avalanche/Source/AvalancheMask/Private/Handling/AvaMaskDesignedMaterialHandle.cpp b/Engine/Plugins/VirtualProduction/Avalanche/Source/AvalancheMask/Private/Handling/AvaMaskDesignedMaterialHandle.cpp
--- a/Engine/Plugins/VirtualProduction/Avalanche/Source/AvalancheMask/Private/Handling/AvaMaskDesignedMaterialHandle.cpp
+++ b/Engine/Plugins/VirtualProduction/Avalanche/Source/AvalancheMask/Private/Handling/AvaMaskDesignedMaterialHandle.cpp
@@ -54,13 +54,7 @@ bool FAvaMaskDesignedMaterialHandle::GetMaskParameters(FAvaMask2DMaterialParamet
 {
 	if (const UMaterialInstanceDynamic* MaterialInstance = WeakMaterialInstance.Get())
 	{
-		OutParameters.SetFromMaterial(
-			MaterialInstance
-			, UE::AvaMask::Internal::TextureParameterInfo
-			, UE::AvaMask::Internal::InvertParameterInfo
-			, UE::AvaMask::Internal::BaseOpacityParameterInfo
-			, UE::AvaMask::Internal::PaddingParameterInfo
-			, UE::AvaMask::Internal::FeatherParameterInfo);
+		OutParameters.StoreFromMaterial(MaterialInstance);
 		return true;
 	}
 
@@ -78,14 +72,9 @@ bool FAvaMaskDesignedMaterialHandle::SetMaskParameters(
 {
 	if (UDynamicMaterialInstance* MaterialInstance = WeakDesignedMaterial.Get())
 	{
-		UE_LOG(LogAvaMask, VeryVerbose, TEXT("SetParameters: Texture:%s"), InTexture ? *InTexture->GetName() : TEXT("(None)"));
+		// note this does no longer deal with textures as the mask texture parameter is deprecated.
+		// and this will not deal with texture parameter collections as this mask handle system is deprecated too.
 
-		UTexture* DefaultTexture = nullptr;
-		MaterialInstance->GetTextureParameterDefaultValue(UE::AvaMask::Internal::TextureParameterInfo, DefaultTexture);
-
-		InTexture = InTexture ? InTexture : DefaultTexture;
-
-		MaterialInstance->SetTextureParameterValueByInfo(UE::AvaMask::Internal::TextureParameterInfo, InTexture);
 		MaterialInstance->SetScalarParameterValueByInfo(UE::AvaMask::Internal::InvertParameterInfo, bInInverted ? 1.0f : 0.0f);
 		MaterialInstance->SetScalarParameterValueByInfo(UE::AvaMask::Internal::BaseOpacityParameterInfo, FMath::Clamp(InBaseOpacity, 0.0f, 1.0f));
 		MaterialInstance->SetVectorParameterValueByInfo(UE::AvaMask::Internal::PaddingParameterInfo, FLinearColor(FMath::Max(0, InPadding.X), FMath::Max(0, InPadding.Y), 0, 0));
@@ -152,13 +141,7 @@ bool FAvaMaskDesignedMaterialHandle::SaveOriginalState(const FStructView& InHand
 	{
 		if (FHandleData* MaterialHandleData = InHandleData.GetPtr<FHandleData>())
 		{
-			MaterialHandleData->OriginalMaskMaterialParameters.SetFromMaterial(
-				DesignedMaterial
-				, UE::AvaMask::Internal::TextureParameterInfo
-				, UE::AvaMask::Internal::InvertParameterInfo
-				, UE::AvaMask::Internal::BaseOpacityParameterInfo
-				, UE::AvaMask::Internal::PaddingParameterInfo
-				, UE::AvaMask::Internal::FeatherParameterInfo);
+			MaterialHandleData->OriginalMaskMaterialParameters.StoreFromMaterial(DesignedMaterial);
 
 			MaterialHandleData->OriginalBlendMode = GetBlendMode();
 
@@ -182,8 +165,7 @@ bool FAvaMaskDesignedMaterialHandle::ApplyOriginalState(const FStructView& InHan
 	{
 		if (const FHandleData* MaterialHandleData = InHandleData.GetPtr<const FHandleData>())
 		{
-			SetMaskParameters(
-				MaterialHandleData->OriginalMaskMaterialParameters.Texture
+			SetMaskParameters(nullptr
 				, MaterialHandleData->OriginalMaskMaterialParameters.BaseOpacity
 				, MaterialHandleData->OriginalMaskMaterialParameters.bInvert
 				, MaterialHandleData->OriginalMaskMaterialParameters.Padding
@@ -213,8 +195,7 @@ bool FAvaMaskDesignedMaterialHandle::ApplyModifiedState(
 	{
 		if (InHandleData.GetPtr<const FHandleData>())
 		{
-			SetMaskParameters(
-				InModifiedParameters.MaterialParameters.Texture
+			SetMaskParameters(nullptr
 				, InModifiedParameters.MaterialParameters.BaseOpacity
 				, InModifiedParameters.MaterialParameters.bInvert
 				, InModifiedParameters.MaterialParameters.Padding
@@ -275,8 +256,7 @@ void FAvaMaskDesignedMaterialHandle::OnMaterialBuilt(UDynamicMaterialModelBase*
 	check(InMaterialModel);
 	
 	// @note: parameter values get reset after build, so we need to re-apply them
-	SetMaskParameters(
-		LastAppliedParameters.MaterialParameters.Texture
+	SetMaskParameters(nullptr
 		, LastAppliedParameters.MaterialParameters.BaseOpacity
 		, LastAppliedParameters.MaterialParameters.bInvert
 		, LastAppliedParameters.MaterialParameters.Padding
diff --git a/Engine/Plugins/VirtualProduction/Avalanche/Source/AvalancheMask/Private/Handling/AvaMaskMaterialInstanceHandle.cpp b/Engine/Plugins/VirtualProduction/Avalanche/Source/AvalancheMask/Private/Handling/AvaMaskMaterialInstanceHandle.cpp
--- a/Engine/Plugins/VirtualProduction/Avalanche/Source/AvalancheMask/Private/Handling/AvaMaskMaterialInstanceHandle.cpp
+++ b/Engine/Plugins/VirtualProduction/Avalanche/Source/AvalancheMask/Private/Handling/AvaMaskMaterialInstanceHandle.cpp
@@ -50,13 +50,7 @@ bool FAvaMaskMaterialInstanceHandle::SaveOriginalState(const FStructView& InHand
 		if (FHandleData* HandleData = InHandleData.GetPtr<FHandleData>())
 		{
 			HandleData->OriginalMaterial = ParentMaterial;
-			HandleData->OriginalMaskMaterialParameters.SetFromMaterial(
-				ParentMaterial
-				, UE::AvaMask::Internal::TextureParameterInfo
-				, UE::AvaMask::Internal::InvertParameterInfo
-				, UE::AvaMask::Internal::BaseOpacityParameterInfo
-				, UE::AvaMask::Internal::PaddingParameterInfo
-				, UE::AvaMask::Internal::FeatherParameterInfo);
+			HandleData->OriginalMaskMaterialParameters.StoreFromMaterial(ParentMaterial);
 			HandleData->OriginalBlendMode = GetBlendMode();
 		}
 
@@ -73,8 +67,7 @@ bool FAvaMaskMaterialInstanceHandle::ApplyOriginalState(
 	if (const FHandleData* HandleData = InHandleData.GetPtr<FHandleData>())
 	{
 		SetBlendMode(HandleData->OriginalBlendMode);
-		SetMaskParameters(
-			HandleData->OriginalMaskMaterialParameters.Texture
+		SetMaskParameters(nullptr
 			, HandleData->OriginalMaskMaterialParameters.BaseOpacity
 			, HandleData->OriginalMaskMaterialParameters.bInvert);
 		
@@ -104,8 +97,7 @@ bool FAvaMaskMaterialInstanceHandle::ApplyModifiedState(
 		if (UMaterialInterface* MaterialInstance = GetOrCreateMaterialInstance(InModifiedParameters.MaterialParameters.CanvasName, TargetBlendMode))
 		{
 			SetBlendMode(InModifiedParameters.MaterialParameters.BlendMode);
-			SetMaskParameters(
-				InModifiedParameters.MaterialParameters.Texture
+			SetMaskParameters(nullptr
 				, InModifiedParameters.MaterialParameters.BaseOpacity
 				, InModifiedParameters.MaterialParameters.bInvert
 			);
diff --git a/Engine/Plugins/VirtualProduction/Avalanche/Source/AvalancheMask/Private/Handling/AvaMaskParametricMaterialHandle.cpp b/Engine/Plugins/VirtualProduction/Avalanche/Source/AvalancheMask/Private/Handling/AvaMaskParametricMaterialHandle.cpp
--- a/Engine/Plugins/VirtualProduction/Avalanche/Source/AvalancheMask/Private/Handling/AvaMaskParametricMaterialHandle.cpp
+++ b/Engine/Plugins/VirtualProduction/Avalanche/Source/AvalancheMask/Private/Handling/AvaMaskParametricMaterialHandle.cpp
@@ -35,13 +35,7 @@ bool FAvaMaskParametricMaterialHandle::GetMaskParameters(FAvaMask2DMaterialParam
 {
 	if (const UMaterialInstanceDynamic* MaterialInstance = WeakMaterialInstance.Get())
 	{
-		OutParameters.SetFromMaterial(
-			MaterialInstance
-			, UE::AvaMask::Internal::TextureParameterInfo
-			, UE::AvaMask::Internal::InvertParameterInfo
-			, UE::AvaMask::Internal::BaseOpacityParameterInfo
-			, UE::AvaMask::Internal::PaddingParameterInfo
-			, UE::AvaMask::Internal::FeatherParameterInfo);
+		OutParameters.StoreFromMaterial(MaterialInstance);
 		return true;
 	}
 
@@ -62,14 +56,11 @@ bool FAvaMaskParametricMaterialHandle::SetMaskParameters(
 		FLinearColor PaddingValueV = FLinearColor(FMath::Max(0, InPadding.X), FMath::Max(0, InPadding.Y), 0, 0);
 		FLinearColor FeatherValueV = FLinearColor(bInApplyFeathering ? 1.0f : 0.0f, FMath::Max(0, InOuterFeathering), FMath::Max(0, InInnerFeathering), FMath::Max(0, FMath::Max(InOuterFeathering, InInnerFeathering)));
 		
-		UE_LOG(LogAvaMask, VeryVerbose, TEXT("SetParameters:\nTexture:%s\n\nFeather:%s\n"), *GetNameSafe(InTexture), *FeatherValueV.ToString());
+		UE_LOG(LogAvaMask, VeryVerbose, TEXT("SetParameters:\nFeather:%s\n"), *FeatherValueV.ToString());
 
-		UTexture* DefaultTexture = nullptr;
-		MaterialInstance->GetTextureParameterDefaultValue(UE::AvaMask::Internal::TextureParameterInfo, DefaultTexture);
+		// note this does no longer deal with textures as the mask texture parameter is deprecated.
+		// and this will not deal with texture parameter collections as this mask handle system is deprecated too.
 
-		InTexture = InTexture ? InTexture : DefaultTexture;
-
-		MaterialInstance->SetTextureParameterValueByInfo(UE::AvaMask::Internal::TextureParameterInfo, InTexture);
 		MaterialInstance->SetScalarParameterValueByInfo(UE::AvaMask::Internal::InvertParameterInfo, bInInverted ? 1.0f : 0.0f);
 		MaterialInstance->SetScalarParameterValueByInfo(UE::AvaMask::Internal::BaseOpacityParameterInfo, FMath::Clamp(InBaseOpacity, 0.0f, 1.0f));
 		MaterialInstance->SetVectorParameterValueByInfo(UE::AvaMask::Internal::PaddingParameterInfo, PaddingValueV);
@@ -142,14 +133,7 @@ bool FAvaMaskParametricMaterialHandle::SaveOriginalState(const FStructView& InHa
 				Material = ParametricMtl->GetDefaultMaterial();
 			}
 
-			MaterialHandleData->OriginalMaskMaterialParameters.SetFromMaterial(
-				Material
-				, UE::AvaMask::Internal::TextureParameterInfo
-				, UE::AvaMask::Internal::InvertParameterInfo
-				, UE::AvaMask::Internal::BaseOpacityParameterInfo
-				, UE::AvaMask::Internal::PaddingParameterInfo
-				, UE::AvaMask::Internal::FeatherParameterInfo);
-
+			MaterialHandleData->OriginalMaskMaterialParameters.StoreFromMaterial(Material);
 			MaterialHandleData->OriginalMaterialStyle = ParametricMtl->GetStyle();
 			MaterialHandleData->OriginalBlendMode = GetBlendMode();
 
@@ -170,8 +154,7 @@ bool FAvaMaskParametricMaterialHandle::ApplyOriginalState(
 		{
 			SetBlendMode(MaterialHandleData->OriginalBlendMode);
 			
-			SetMaskParameters(
-				MaterialHandleData->OriginalMaskMaterialParameters.Texture
+			SetMaskParameters(nullptr
 				, MaterialHandleData->OriginalMaskMaterialParameters.BaseOpacity
 				, MaterialHandleData->OriginalMaskMaterialParameters.bInvert
 				, MaterialHandleData->OriginalMaskMaterialParameters.Padding
@@ -202,8 +185,7 @@ bool FAvaMaskParametricMaterialHandle::ApplyModifiedState(
 		{
 			SetBlendMode(InModifiedParameters.MaterialParameters.BlendMode);
 
-			SetMaskParameters(
-				InModifiedParameters.MaterialParameters.Texture
+			SetMaskParameters(nullptr
 				, InModifiedParameters.MaterialParameters.BaseOpacity
 				, InModifiedParameters.MaterialParameters.bInvert
 				, InModifiedParameters.MaterialParameters.Padding
diff --git a/Engine/Plugins/VirtualProduction/Avalanche/Source/AvalancheMask/Private/Handling/IAvaMaskMaterialHandle.h b/Engine/Plugins/VirtualProduction/Avalanche/Source/AvalancheMask/Private/Handling/IAvaMaskMaterialHandle.h
--- a/Engine/Plugins/VirtualProduction/Avalanche/Source/AvalancheMask/Private/Handling/IAvaMaskMaterialHandle.h
+++ b/Engine/Plugins/VirtualProduction/Avalanche/Source/AvalancheMask/Private/Handling/IAvaMaskMaterialHandle.h
@@ -97,13 +97,7 @@ bool TAvaMaskMaterialHandle<HandleDataType>::GetMaskParameters(FAvaMask2DMateria
 {
 	if (const UMaterialInstanceDynamic* MaterialInstance = GetMaterialInstance())
 	{
-		OutParameters.SetFromMaterial(
-			MaterialInstance
-			, UE::AvaMask::Internal::TextureParameterInfo
-			, UE::AvaMask::Internal::InvertParameterInfo
-			, UE::AvaMask::Internal::BaseOpacityParameterInfo
-			, UE::AvaMask::Internal::PaddingParameterInfo
-			, UE::AvaMask::Internal::FeatherParameterInfo);
+		OutParameters.StoreFromMaterial(MaterialInstance);
 		return true;
 	}
 
@@ -121,14 +115,11 @@ bool TAvaMaskMaterialHandle<HandleDataType>::SetMaskParameters(UTexture* InTextu
 {
 	if (UMaterialInstanceDynamic* MaterialInstance = GetMaterialInstance())
 	{
-		UE_LOG(LogAvaMask, Verbose, TEXT("SetParameters: Texture:%s, Material:%s"), InTexture ? *InTexture->GetName() : TEXT("(None)"), *MaterialInstance->GetName());
+		UE_LOG(LogAvaMask, Verbose, TEXT("SetParameters: Material:%s"), *MaterialInstance->GetName());
 
-		UTexture* DefaultTexture = nullptr;
-		MaterialInstance->GetTextureParameterDefaultValue(UE::AvaMask::Internal::TextureParameterInfo, DefaultTexture);
+		// note this does no longer deal with textures as the mask texture parameter is deprecated.
+		// and this will not deal with texture parameter collections as this mask handle system is deprecated too.
 
-		InTexture = InTexture ? InTexture : DefaultTexture;
-
-		MaterialInstance->SetTextureParameterValueByInfo(UE::AvaMask::Internal::TextureParameterInfo, InTexture);
 		MaterialInstance->SetScalarParameterValueByInfo(UE::AvaMask::Internal::InvertParameterInfo, bInInverted ? 1.0f : 0.0f);
 		MaterialInstance->SetScalarParameterValueByInfo(UE::AvaMask::Internal::BaseOpacityParameterInfo, FMath::Clamp(InBaseOpacity, 0.0f, 1.0f));
 		MaterialInstance->SetVectorParameterValueByInfo(UE::AvaMask::Internal::PaddingParameterInfo, FLinearColor(FMath::Max(0, InPadding.X), FMath::Max(0, InPadding.Y), 0, 0));
@@ -147,10 +138,9 @@ bool TAvaMaskMaterialHandle<HandleDataType>::HasRequiredParameters(TArray<FStrin
 	{
 		return UE::Ava::MaterialHasParameters(
 			*Material
-			, {
-				{ UE::AvaMask::Internal::TextureParameterInfo.Name, EMaterialParameterType::Texture }
-				, { UE::AvaMask::Internal::BaseOpacityParameterInfo.Name, EMaterialParameterType::Scalar }
-				, { UE::AvaMask::Internal::InvertParameterInfo.Name, EMaterialParameterType::Scalar }
+			, { 
+				{ UE::AvaMask::Internal::BaseOpacityParameterInfo.Name, EMaterialParameterType::Scalar }, 
+				{ UE::AvaMask::Internal::InvertParameterInfo.Name, EMaterialParameterType::Scalar }
 			}
 			, OutMissingParameterNames);
 	}
diff --git a/Engine/Plugins/VirtualProduction/Avalanche/Source/AvalancheMask/Private/Mask2D/AvaMask2DBaseModifier.cpp b/Engine/Plugins/VirtualProduction/Avalanche/Source/AvalancheMask/Private/Mask2D/AvaMask2DBaseModifier.cpp
--- a/Engine/Plugins/VirtualProduction/Avalanche/Source/AvalancheMask/Private/Mask2D/AvaMask2DBaseModifier.cpp
+++ b/Engine/Plugins/VirtualProduction/Avalanche/Source/AvalancheMask/Private/Mask2D/AvaMask2DBaseModifier.cpp
@@ -171,17 +171,8 @@ void UAvaMask2DBaseModifier::OnCanvasChanged()
 void UAvaMask2DBaseModifier::UpdateCanvas()
 {
 	// Reset/invalidate cached
-	{
-PRAGMA_DISABLE_DEPRECATION_WARNINGS
-		for (TPair<TWeakObjectPtr<AActor>, FAvaMask2DActorData>& ActorDataPair : ActorData_DEPRECATED)
-		{
-			ActorDataPair.Value.CanvasTextureWeak.Reset();
-		}
-PRAGMA_ENABLE_DEPRECATION_WARNINGS
-
-		LastResolvedCanvasName = NAME_None;
-		CanvasWeak.Reset();
-	}
+	LastResolvedCanvasName = NAME_None;
+	CanvasWeak.Reset();
 
 #if WITH_EDITOR
 	if (!Channel.IsNone())
@@ -309,10 +300,10 @@ void UAvaMask2DBaseModifier::OnModifierRemoved(EActorModifierCoreDisableReason I
 PRAGMA_DISABLE_DEPRECATION_WARNINGS
 	ActorDataActors.Reserve(ActorData_DEPRECATED.Num());	
 	Algo::Transform(ActorData_DEPRECATED, ActorDataActors, [](const TPair<TWeakObjectPtr<AActor>, FAvaMask2DActorData>& InActorDataPair)
-PRAGMA_ENABLE_DEPRECATION_WARNINGS
 	{
 		return InActorDataPair.Key.Get();
 	});
+PRAGMA_ENABLE_DEPRECATION_WARNINGS
 	
 	for (AActor* Actor : ActorDataActors)
 	{
@@ -604,18 +595,6 @@ void UAvaMask2DBaseModifier::OnUseParentChannelChanged()
 
 void UAvaMask2DBaseModifier::OnMaskSetCanvas(const UGeometryMaskCanvas* InCanvas, AActor* InActor)
 {
-	if (!InCanvas)
-	{
-		return;
-	}
-
-PRAGMA_DISABLE_DEPRECATION_WARNINGS
-	if (FAvaMask2DActorData* StoredActorData = ActorData_DEPRECATED.Find(InActor))
-PRAGMA_ENABLE_DEPRECATION_WARNINGS
-	{
-		// Cache canvas texture
-		StoredActorData->CanvasTextureWeak = InCanvas->GetTexture();
-	}
 }
 
 void UAvaMask2DBaseModifier::TryResolveCanvas()
@@ -644,39 +623,6 @@ void UAvaMask2DBaseModifier::TryResolveCanvas()
 
 UTexture* UAvaMask2DBaseModifier::TryResolveCanvasTexture(AActor* InActor, FAvaMask2DActorData& InActorData)
 {
-	if (!InActor)
-	{
-		return nullptr;
-	}
-
-	// Get cached version
-	if (UTexture* CanvasTexture = InActorData.CanvasTextureWeak.Get())
-	{
-		return CanvasTexture;
-	}
-
-	const ULevel* Level = InActor->GetLevel();
-	if (!Level || !Level->OwningWorld)
-	{
-		return nullptr;
-	}
-
-	// Try get canvas texture if already available
-	UGeometryMaskWorldSubsystem* MaskSubsystem = Level->OwningWorld->GetSubsystem<UGeometryMaskWorldSubsystem>();
-	if (!MaskSubsystem)
-	{
-		return nullptr;
-	}
-
-	if (const UGeometryMaskCanvas* Canvas = MaskSubsystem->GetNamedCanvas(Level, Channel))
-	{
-		// Cache canvas texture
-		UTexture* CanvasTexture = Canvas->GetTexture();
-		InActorData.CanvasTextureWeak = CanvasTexture;
-
-		return CanvasTexture;
-	}
-
 	return nullptr;
 }
 
diff --git a/Engine/Plugins/VirtualProduction/Avalanche/Source/AvalancheMask/Private/Mask2D/AvaMask2DModifierShared.cpp b/Engine/Plugins/VirtualProduction/Avalanche/Source/AvalancheMask/Private/Mask2D/AvaMask2DModifierShared.cpp
new file mode 100644
--- /dev/null
+++ b/Engine/Plugins/VirtualProduction/Avalanche/Source/AvalancheMask/Private/Mask2D/AvaMask2DModifierShared.cpp
@@ -0,0 +1,101 @@
+// Copyright Epic Games, Inc. All Rights Reserved.
+
+#include "Mask2D/AvaMask2DModifierShared.h"
+#include "Mask2D/AvaMask2DReadModifier.h"
+
+void UAvaMask2DReadModifierShared::RegisterModifier(TNotNull<UAvaMask2DReadModifier*> InModifier)
+{
+	const AActor* ModifiedActor = InModifier->GetModifiedActor();
+	if (!ModifiedActor)
+	{
+		return;
+	}
+
+	FAvaMask2DReadModifierGroup& ModifierGroup = ModifierGroups.FindOrAdd(ModifiedActor);
+	ModifierGroup.Modifiers.AddUnique(InModifier);
+}
+
+void UAvaMask2DReadModifierShared::UnregisterModifier(TNotNull<UAvaMask2DReadModifier*> InModifier)
+{
+	const AActor* ModifiedActor = InModifier->GetModifiedActor();
+	if (!ModifiedActor)
+	{
+		return;
+	}
+
+	const TObjectKey<AActor> ActorKey(ModifiedActor);
+	if (FAvaMask2DReadModifierGroup* ModifierGroup = ModifierGroups.Find(ActorKey))
+	{
+		ModifierGroup->Modifiers.Remove(InModifier);
+		if (ModifierGroup->Modifiers.IsEmpty())
+		{
+			ModifierGroups.Remove(ActorKey);
+		}
+	}
+}
+
+int32 UAvaMask2DReadModifierShared::GatherCanvases(TNotNull<UAvaMask2DReadModifier*> InModifier, TAdderReserverRef<TNotNull<const UGeometryMaskCanvas*>> InCanvases) const
+{
+	int32 CanvasAddedCount = 0;
+
+	// limit already reached. Can't add more
+	if (InCanvases.Num() >= MaxCanvasGatherCount)
+	{
+		return CanvasAddedCount;
+	}
+
+	// adds the canvas to the container ref and reduces the mask count 
+	auto AddCanvas = [&InCanvases, &CanvasAddedCount](UGeometryMaskCanvas* InCanvas)
+		{
+			if (InCanvas)
+			{
+				InCanvases.Add(InCanvas);
+				++CanvasAddedCount;
+			}
+			return InCanvases.Num() < MaxCanvasGatherCount;
+		};
+
+	if (!AddCanvas(InModifier->GetCurrentCanvas()))
+	{
+		return CanvasAddedCount;
+	}
+
+	// stop at just adding the current canvas if nested masking is not enabled
+	if (!InModifier->IsNestedMaskingEnabled())
+	{
+		return CanvasAddedCount;
+	}
+
+	const AActor* Actor = InModifier->GetModifiedActor();
+	while (Actor)
+	{
+		if (const FAvaMask2DReadModifierGroup* ModifierGroup = ModifierGroups.Find(Actor))
+		{
+			for (UAvaMask2DReadModifier* IteratedModifier : ModifierGroup->Modifiers)
+			{
+				// skip invalid modifier and the parameter modifier as it has already been processed
+				if (IteratedModifier && IteratedModifier != InModifier)
+				{
+					if (!AddCanvas(IteratedModifier->GetCurrentCanvas()))
+					{
+						return CanvasAddedCount;
+					}
+				}
+			}
+		}
+		Actor = Actor->GetAttachParentActor();
+	}
+
+	return CanvasAddedCount;
+}
+
+void UAvaMask2DReadModifierShared::AddReferencedObjects(UObject* InThis, FReferenceCollector& Collector)
+{
+	Super::AddReferencedObjects(InThis, Collector);
+
+	UAvaMask2DReadModifierShared* const This = CastChecked<UAvaMask2DReadModifierShared>(InThis);
+	for (TPair<TObjectKey<AActor>, FAvaMask2DReadModifierGroup>& Pair : This->ModifierGroups)
+	{
+		Collector.AddPropertyReferencesWithStructARO(FAvaMask2DReadModifierGroup::StaticStruct(), &Pair.Value, This);
+	}
+}
diff --git a/Engine/Plugins/VirtualProduction/Avalanche/Source/AvalancheMask/Private/Mask2D/AvaMask2DReadModifier.cpp b/Engine/Plugins/VirtualProduction/Avalanche/Source/AvalancheMask/Private/Mask2D/AvaMask2DReadModifier.cpp
--- a/Engine/Plugins/VirtualProduction/Avalanche/Source/AvalancheMask/Private/Mask2D/AvaMask2DReadModifier.cpp
+++ b/Engine/Plugins/VirtualProduction/Avalanche/Source/AvalancheMask/Private/Mask2D/AvaMask2DReadModifier.cpp
@@ -10,6 +10,7 @@
 #include "Engine/Engine.h"
 #include "GeometryMaskCanvas.h"
 #include "GeometryMaskCanvasResource.h"
+#include "Mask2D/AvaMask2DModifierShared.h"
 #include "MaterialBridge/AvaMaterialBridge.h"
 #include "MaterialBridge/AvaMaterialBridgeRegistry.h"
 #include "MaterialBridge/Context/AvaMaterialBridgeApplyStateContext.h"
@@ -33,7 +34,8 @@
 #if WITH_EDITOR
 const TAvaPropertyChangeDispatcher<UAvaMask2DReadModifier> UAvaMask2DReadModifier::PropertyChangeDispatcher =
 {
-	{ GET_MEMBER_NAME_CHECKED(UAvaMask2DReadModifier, BaseOpacity), &UAvaMask2DReadModifier::OnBaseOpacityChanged }
+	{ GET_MEMBER_NAME_CHECKED(UAvaMask2DReadModifier, BaseOpacity), &UAvaMask2DReadModifier::OnBaseOpacityChanged },
+	{ GET_MEMBER_NAME_CHECKED(UAvaMask2DReadModifier, bNestedMaskingEnabled), &UAvaMask2DReadModifier::OnNestedMaskingEnabledChanged }
 };
 #endif
 
@@ -55,6 +57,15 @@ void UAvaMask2DReadModifier::PostEditChangeProperty(FPropertyChangedEvent& InPro
 	UAvaMask2DReadModifier::PropertyChangeDispatcher.OnPropertyChanged(this, InPropertyChangedEvent);
 }
 
+void UAvaMask2DReadModifier::SetNestedMaskingEnabled(const bool bInNestedMaskingEnabled)
+{
+	if (bNestedMaskingEnabled != bInNestedMaskingEnabled)
+	{
+		bNestedMaskingEnabled = bInNestedMaskingEnabled;
+		OnNestedMaskingEnabledChanged();
+	}
+}
+
 void UAvaMask2DReadModifier::OnMaterialCompiled(UMaterialInterface* InMaterial)
 {
 	if (!InMaterial)
@@ -85,13 +96,19 @@ void UAvaMask2DReadModifier::Apply()
 {
 	Super::Apply();
 
-	if (UGeometryMaskCanvas* Canvas = GetCurrentCanvas())
+	if (UAvaMask2DReadModifierShared* Shared = GetShared<UAvaMask2DReadModifierShared>())
 	{
-		for (TPair<TObjectPtr<UObject>, TInstancedStruct<FAvaMaterialContainerState>>& Pair : MaterialContainerStates)
+		TArray<TNotNull<const UGeometryMaskCanvas*>, TInlineAllocator<1>> Canvases;
+		Shared->GatherCanvases(this, Canvases);
+
+		if (!Canvases.IsEmpty())
 		{
-			if (Pair.Key)
+			for (TPair<TObjectPtr<UObject>, TInstancedStruct<FAvaMaterialContainerState>>& Pair : MaterialContainerStates)
 			{
-				ApplyRead(Pair.Key, Canvas);
+				if (Pair.Key)
+				{
+					ApplyRead(Pair.Key, Canvases);
+				}
 			}
 		}
 	}
@@ -156,8 +173,13 @@ EBlendMode UAvaMask2DReadModifier::GetBlendMode() const
 	return bUseBlur || bUseFeathering ? EBlendMode::BLEND_Translucent : EBlendMode::BLEND_Masked;
 }
 
-bool UAvaMask2DReadModifier::ApplyRead(TNotNull<UObject*> InMaterialContainer, TNotNull<const UGeometryMaskCanvas*> InCurrentCanvas)
+bool UAvaMask2DReadModifier::ApplyRead(TNotNull<UObject*> InMaterialContainer, TConstArrayView<TNotNull<const UGeometryMaskCanvas*>> InCanvases)
 {
+	if (InCanvases.IsEmpty())
+	{
+		return true;
+	}
+
 	const UE::Ava::FMaterialBridge* const MaterialBridge = UE::Ava::FMaterialBridgeRegistry::Get().GetMaterialBridge(UE::Ava::FConstDataView(InMaterialContainer));
 	if (!MaterialBridge)
 	{
@@ -165,24 +187,45 @@ bool UAvaMask2DReadModifier::ApplyRead(TNotNull<UObject*> InMaterialContainer, T
 		return true;
 	}
 
-	FGeometryMaskDrawingContext DrawingContext(InCurrentCanvas->GetCanvasId().Level);
-	const FIntPoint ViewportSize = InCurrentCanvas->GetResource()->GetMaxViewportSize();
-	const FVector2f ViewportPadding = FVector2f::One() - FVector2f(ViewportSize + InCurrentCanvas->GetResource()->GetViewportPadding(DrawingContext)) / ViewportSize;
+	const UGeometryMaskCanvasResource* const FirstCanvasResource = InCanvases[0]->GetResource();
+	if (!ensure(FirstCanvasResource))
+	{
+		return true;
+	}
+
+	// using only the first canvas parameters as it is consistent across the other canvases.
+	const FIntPoint ViewportSize = FirstCanvasResource->GetMaxViewportSize();
+	const FVector2f ViewportPadding = FVector2f::One() - FVector2f(ViewportSize + FirstCanvasResource->GetViewportPadding()) / ViewportSize;
 
 	FAvaMask2DMaskState MaskedState;
 
 #if WITH_EDITOR
 	MaskedState.OutputProcessor = GetDefault<UAvaMaskSettings>()->GetMaterialFunction();
 #endif
-	MaskedState.MaterialParameters.BlendMode = GetBlendMode();		
-	MaskedState.MaterialParameters.CanvasName = InCurrentCanvas->GetCanvasName();
-	MaskedState.MaterialParameters.Texture = InCurrentCanvas->GetTexture();
+	MaskedState.MaterialParameters.BlendMode = GetBlendMode();
 	MaskedState.MaterialParameters.BaseOpacity = BaseOpacity;
 	MaskedState.MaterialParameters.bInvert = bInverted;
 	MaskedState.MaterialParameters.Padding = ViewportPadding;
 	MaskedState.MaterialParameters.bApplyFeathering = bUseFeathering;
 	MaskedState.MaterialParameters.OuterFeatherRadius = OuterFeatherRadius;
 	MaskedState.MaterialParameters.InnerFeatherRadius = InnerFeatherRadius;
+	MaskedState.MaterialParameters.RenderTarget = FirstCanvasResource->GetRenderTarget();
+
+	// Fill the mask indices of each of the canvases (max 4 because of linear color)
+	static_assert(UAvaMask2DReadModifierShared::MaxCanvasGatherCount == 4);
+	for (int32 Index = 0; Index < UAvaMask2DReadModifierShared::MaxCanvasGatherCount; ++Index)
+	{
+		float& MaskIndexComponent = MaskedState.MaterialParameters.MaskIndices.Component(Index);
+		if (InCanvases.IsValidIndex(Index) && ensure(InCanvases[Index]->GetRenderTarget() == MaskedState.MaterialParameters.RenderTarget))
+		{
+			const int16 SliceIndex = InCanvases[Index]->GetRenderTargetSliceIndex();
+			MaskIndexComponent = static_cast<float>(SliceIndex);
+		}
+		else
+		{
+			MaskIndexComponent = -1.f;
+		}
+	}
 
 	MaterialBridge->AccessSlots(UE::Ava::FMaterialBridgeWriteSlotContext(InMaterialContainer)
 		, [This=this, &MaskedState](const UE::Ava::FMaterialBridgeWriteSlotContext& InContext, UE::Ava::FMaterialBridgeWriteSlot& InSlot)->UE::Ava::EControlFlow
@@ -253,7 +296,7 @@ bool UAvaMask2DReadModifier::ApplyRead(AActor* InActor, FAvaMask2DActorData& InA
 	{
 		if (UGeometryMaskCanvas* Canvas = GetCurrentCanvas())
 		{
-			return ApplyRead(InActor, Canvas);
+			return ApplyRead(InActor, { Canvas });
 		}
 	}
 	// Ok to continue, just nothing done for this actor
@@ -278,6 +321,11 @@ void UAvaMask2DReadModifier::OnModifierEnabled(EActorModifierCoreEnableReason In
 {
 	Super::OnModifierEnabled(InReason);
 
+	if (UAvaMask2DReadModifierShared* Shared = GetShared<UAvaMask2DReadModifierShared>(/*bInCreateIfNone*/true))
+	{
+		Shared->RegisterModifier(this);
+	}
+
 #if WITH_EDITOR
 	UMaterial::OnMaterialCompilationFinished().AddUObject(this, &UAvaMask2DReadModifier::OnMaterialCompiled);
 	FEditorDelegates::EndPIE.AddUObject(this, &UAvaMask2DReadModifier::OnPIEEnded);
@@ -288,6 +336,11 @@ void UAvaMask2DReadModifier::OnModifierDisabled(EActorModifierCoreDisableReason
 {
 	Super::OnModifierDisabled(InReason);
 
+	if (UAvaMask2DReadModifierShared* Shared = GetShared<UAvaMask2DReadModifierShared>(/*bInCreateIfNone*/false))
+	{
+		Shared->UnregisterModifier(this);
+	}
+
 #if WITH_EDITOR
 	UMaterial::OnMaterialCompilationFinished().RemoveAll(this);
 	FEditorDelegates::EndPIE.RemoveAll(this);
@@ -388,6 +441,11 @@ void UAvaMask2DReadModifier::OnBaseOpacityChanged()
 	MarkModifierDirty();
 }
 
+void UAvaMask2DReadModifier::OnNestedMaskingEnabledChanged()
+{
+	MarkModifierDirty();
+}
+
 void UAvaMask2DReadModifier::SetupMaskComponent(UActorComponent* InComponent)
 {
 }
diff --git a/Engine/Plugins/VirtualProduction/Avalanche/Source/AvalancheMask/Public/AvaMaskTypes.h b/Engine/Plugins/VirtualProduction/Avalanche/Source/AvalancheMask/Public/AvaMaskTypes.h
--- a/Engine/Plugins/VirtualProduction/Avalanche/Source/AvalancheMask/Public/AvaMaskTypes.h
+++ b/Engine/Plugins/VirtualProduction/Avalanche/Source/AvalancheMask/Public/AvaMaskTypes.h
@@ -9,6 +9,7 @@
 class UMaterialInstanceDynamic;
 class UMaterialInterface;
 class UTexture;
+class UTextureRenderTarget2DArray;
 
 /** Describes a component/material slot pair. */
 USTRUCT()
@@ -55,8 +56,17 @@ PRAGMA_ENABLE_DEPRECATION_WARNINGS
 	FName CanvasName;
 
 	/** Mask texture. */
+	UE_DEPRECATED(5.8, "Mask texture is deprecated. Use the mask texture collection member variable instead")
 	UPROPERTY()
-	TObjectPtr<UTexture> Texture = nullptr;
+	TObjectPtr<UTexture> Texture_DEPRECATED = nullptr;
+
+	/** Render target array to use for masking */
+	UPROPERTY()
+	TObjectPtr<UTextureRenderTarget2DArray> RenderTarget;
+
+	/** Indices to the mask textures array. Each component (R,G,B,A) represents an index */
+	UPROPERTY()
+	FLinearColor MaskIndices = FLinearColor(-1, -1, -1, -1);
 
 	/** Multiplies the mask texture to determine which channel to read from. */
 	UE_DEPRECATED(5.8, "Color channel management has been deprecated as it's no longer needed")
@@ -92,15 +102,11 @@ PRAGMA_ENABLE_DEPRECATION_WARNINGS
 	UPROPERTY()
 	TEnumAsByte<EBlendMode> BlendMode = EBlendMode::BLEND_Opaque;
 
+	/** Applies the current values of this struct to the given material parameters. Returns true if successful. */
 	bool ApplyToMID(UMaterialInstanceDynamic* InMaterial) const;
-	bool StoreFromMaterial(const UMaterialInterface* InMaterial);
-
-	/** Set current values from the input material and named parameters, return true if successful. */
-	UE_DEPRECATED(5.8, "Color channel management has been deprecated as it's no longer needed. Use the overload without the Channel Parameter")
-	bool SetFromMaterial(const UMaterialInterface* InMaterial, const FMaterialParameterInfo& InTextureParameter, const FMaterialParameterInfo& InChannelParameter, const FMaterialParameterInfo& InInvertParameter, const FMaterialParameterInfo& InBaseOpacityParameter, const FMaterialParameterInfo& InPaddingParameter, const FMaterialParameterInfo& InFeatherParameter);
 
-	/** Set current values from the input material and named parameters, return true if successful. */
-	bool SetFromMaterial(const UMaterialInterface* InMaterial, const FMaterialParameterInfo& InTextureParameter, const FMaterialParameterInfo& InInvertParameter, const FMaterialParameterInfo& InBaseOpacityParameter, const FMaterialParameterInfo& InPaddingParameter, const FMaterialParameterInfo& InFeatherParameter);
+	/** Set current values of the given material and stores it in the member variables of this struct. Returns true if successful. */
+	bool StoreFromMaterial(const UMaterialInterface* InMaterial);
 };
 
 /** Encapsulates all parameters to apply to a given modifier subject. */
diff --git a/Engine/Plugins/VirtualProduction/Avalanche/Source/AvalancheMask/Public/Mask2D/AvaMask2DBaseModifier.h b/Engine/Plugins/VirtualProduction/Avalanche/Source/AvalancheMask/Public/Mask2D/AvaMask2DBaseModifier.h
--- a/Engine/Plugins/VirtualProduction/Avalanche/Source/AvalancheMask/Public/Mask2D/AvaMask2DBaseModifier.h
+++ b/Engine/Plugins/VirtualProduction/Avalanche/Source/AvalancheMask/Public/Mask2D/AvaMask2DBaseModifier.h
@@ -35,8 +35,11 @@ enum class EAvaMask2DMode : uint8
 	Write		UMETA(DisplayName = "Source", ToolTip = "Use the specified Mask Channel to render this geometry to"),
 };
 
-/** Used for target actors to store essential information */
-USTRUCT()
+/**
+ * DEPRECATED for 5.8
+ * Used for target actors to store essential information 
+ */
+USTRUCT(meta=(Deprecated="5.8"))
 struct FAvaMask2DActorData
 {
 	GENERATED_BODY()
@@ -56,7 +59,7 @@ class UAvaMask2DBaseModifier : public UActorModifierArrangeBaseModifier
 {
 	GENERATED_BODY()
 
-	friend class UAvaMask2DModifierShared;
+	friend class UAvaMask2DReadModifierShared;
 	
 public:
 	//~ Begin UObject
@@ -177,6 +180,8 @@ protected:
 	UE_API bool TryResolveParentChannel();
 
 	UE_API void TryResolveCanvas();
+
+	UE_DEPRECATED(5.8, "ActorData has been deprecated. TryResolveCanvasTexture with actor data is no longer supported.")
 	UE_API UTexture* TryResolveCanvasTexture(AActor* InActor, FAvaMask2DActorData& InActorData);
 
 	UE_DEPRECATED(5.8, "Material Mask Handles have been deprecated. Material Bridge is used instead.")
@@ -246,7 +251,7 @@ protected:
 
 	/** Cached actor data to apply/restore */
 	UE_DEPRECATED(5.8, "ActorData has been deprecated. Material Container State is used instead.")
-	UPROPERTY(DuplicateTransient)
+	UPROPERTY(DuplicateTransient, meta=(DeprecatedProperty))
 	TMap<TWeakObjectPtr<AActor>, FAvaMask2DActorData> ActorData_DEPRECATED;
 
 	/** Reference to the last resolved canvas name, as stored in CanvasWeak */
diff --git a/Engine/Plugins/VirtualProduction/Avalanche/Source/AvalancheMask/Public/Mask2D/AvaMask2DModifierShared.h b/Engine/Plugins/VirtualProduction/Avalanche/Source/AvalancheMask/Public/Mask2D/AvaMask2DModifierShared.h
new file mode 100644
--- /dev/null
+++ b/Engine/Plugins/VirtualProduction/Avalanche/Source/AvalancheMask/Public/Mask2D/AvaMask2DModifierShared.h
@@ -0,0 +1,55 @@
+// Copyright Epic Games, Inc. All Rights Reserved.
+
+#pragma once
+
+#include "Containers/AdderRef.h"
+#include "Modifiers/ActorModifierCoreSharedObject.h"
+#include "UObject/ObjectKey.h"
+#include "AvaMask2DModifierShared.generated.h"
+
+class AActor;
+class UAvaMask2DReadModifier;
+class UGeometryMaskCanvas;
+
+/** Represents a group of the modifiers in an actor */
+USTRUCT()
+struct FAvaMask2DReadModifierGroup
+{
+	GENERATED_BODY()
+
+	/** All the modifiers found for an actor */
+	UPROPERTY()
+	TArray<TObjectPtr<UAvaMask2DReadModifier>> Modifiers;
+};
+
+/** Shared object for mask read modifiers to share data between each other */
+UCLASS(Hidden, MinimalAPI)
+class UAvaMask2DReadModifierShared : public UActorModifierCoreSharedObject
+{
+	GENERATED_BODY()
+
+public:
+	/** A maximum of four canvases can be gathered as only four mask indices are forwarded over (via an FLinearColor). */
+	static constexpr int32 MaxCanvasGatherCount = 4;
+
+	/** Adds the modifier to its actor group */
+	void RegisterModifier(TNotNull<UAvaMask2DReadModifier*> InModifier);
+	
+	/** Removes the modifier from its actor group */
+	void UnregisterModifier(TNotNull<UAvaMask2DReadModifier*> InModifier);
+
+	/**
+	 * Gathers all the canvases relevant to the modifier.
+	 * If nesting is enabled, the first four canvases of modifiers found in the attach parent too.
+	 * If nesting is disabled, only the modifier's own canvas will be gathered.
+	 * @see MaxCanvasGatherCount
+	 */
+	int32 GatherCanvases(TNotNull<UAvaMask2DReadModifier*> InModifier, TAdderReserverRef<TNotNull<const UGeometryMaskCanvas*>> InCanvases) const;
+
+	/** Adds the referenced objects in modifier groups */
+	static void AddReferencedObjects(UObject* InThis, FReferenceCollector& Collector);
+
+private:
+	/** Read modifiers grouped by actors */
+	TMap<TObjectKey<AActor>, FAvaMask2DReadModifierGroup> ModifierGroups;
+};
diff --git a/Engine/Plugins/VirtualProduction/Avalanche/Source/AvalancheMask/Public/Mask2D/AvaMask2DReadModifier.h b/Engine/Plugins/VirtualProduction/Avalanche/Source/AvalancheMask/Public/Mask2D/AvaMask2DReadModifier.h
--- a/Engine/Plugins/VirtualProduction/Avalanche/Source/AvalancheMask/Public/Mask2D/AvaMask2DReadModifier.h
+++ b/Engine/Plugins/VirtualProduction/Avalanche/Source/AvalancheMask/Public/Mask2D/AvaMask2DReadModifier.h
@@ -43,7 +43,7 @@ class UAvaMask2DReadModifier : public UAvaMask2DBaseModifier
 {
 	GENERATED_BODY()
 
-	friend class UAvaMask2DModifierShared;
+	friend class UAvaMask2DReadModifierShared;
 	
 public:
 	//~ Begin UObject
@@ -52,6 +52,12 @@ public:
 #endif
 	//~ End UObject
 
+	bool IsNestedMaskingEnabled() const
+	{
+		return bNestedMaskingEnabled;
+	}
+	UE_API void SetNestedMaskingEnabled(const bool bInNestedMaskingEnabled);
+
 	float GetBaseOpacity() const
 	{
 		return BaseOpacity;
@@ -77,6 +83,7 @@ protected:
 	UE_API EBlendMode GetBlendMode() const;
 
 	UE_API void OnBaseOpacityChanged();
+	UE_API void OnNestedMaskingEnabledChanged();
 
 	UE_DEPRECATED(5.8, "SetupMaskReadComponent is deprecated. Mask Read Modifier no longer allocates a component")
 	UE_API virtual void SetupMaskComponent(UActorComponent* InComponent) override;
@@ -84,19 +91,23 @@ protected:
 	UE_DEPRECATED(5.8, "SetupMaskReadComponent is deprecated. Mask Read Modifier no longer allocates a component")
 	UE_API void SetupMaskReadComponent(IGeometryMaskReadInterface* InMaskReader);
 
-	UE_API bool ApplyRead(TNotNull<UObject*> InMaterialContainer, TNotNull<const UGeometryMaskCanvas*> InCurrentCanvas);
+	UE_API bool ApplyRead(TNotNull<UObject*> InMaterialContainer, TConstArrayView<TNotNull<const UGeometryMaskCanvas*>> InCanvases);
 
 	void ApplyMaskState(const UE::Ava::FMaterialBridgeWriteSlotContext& InContext, UE::Ava::FMaterialBridgeWriteSlot& InSlot, const FAvaMask2DMaskState& InMaskState);
 
-	UE_DEPRECATED(5.8, "Use ApplyRead that takes in a canvas instead")
+	UE_DEPRECATED(5.8, "Use ApplyRead that takes in an array of canvas instead")
 	UE_API bool ApplyRead(AActor* InActor, FAvaMask2DActorData& InActorData);
 
-	UE_DEPRECATED(5.8, "Use ApplyRead that takes in a canvas instead")
+	UE_DEPRECATED(5.8, "Use ApplyRead that takes in an array of canvas instead")
 	UE_API bool ApplyRead(AActor* InActor, FAvaMask2DActorData& InActorData, FText& OutFailReason);
 
 	/** Stores the container and material states for a given actor */
 	void StoreActorState(AActor* InActor);
 
+	/** Whether this modifier should also take into account parent masks */
+	UPROPERTY(EditInstanceOnly, Getter = "IsNestedMaskingEnabled", Setter = "SetNestedMaskingEnabled", Category = "Mask2D", meta = (AllowPrivateAccess = "true"))
+	bool bNestedMaskingEnabled = false;
+
 	/** Base opacity/alpha to use in Read mode */
 	UPROPERTY(EditInstanceOnly, BlueprintReadWrite, Getter, Setter, Category = "Mask2D", meta = (ClampMin = 0.0, ClampMax = 1.0, AllowPrivateAccess = "true"))
 	float BaseOpacity = 0.0f;
diff --git a/Engine/Plugins/VirtualProduction/Avalanche/Source/AvalancheMask/Public/Mask2D/AvaMask2DWriteModifier.h b/Engine/Plugins/VirtualProduction/Avalanche/Source/AvalancheMask/Public/Mask2D/AvaMask2DWriteModifier.h
--- a/Engine/Plugins/VirtualProduction/Avalanche/Source/AvalancheMask/Public/Mask2D/AvaMask2DWriteModifier.h
+++ b/Engine/Plugins/VirtualProduction/Avalanche/Source/AvalancheMask/Public/Mask2D/AvaMask2DWriteModifier.h
@@ -38,7 +38,7 @@ class  UAvaMask2DWriteModifier : public UAvaMask2DBaseModifier
 {
 	GENERATED_BODY()
 
-	friend class UAvaMask2DModifierShared;
+	friend class UAvaMask2DReadModifierShared;
 	
 public:
 	EGeometryMaskCompositeOperation GetWriteOperation() const
diff --git a/Engine/Plugins/VirtualProduction/GeometryMask/Source/GeometryMask/Private/GeometryMaskCanvas.cpp b/Engine/Plugins/VirtualProduction/GeometryMask/Source/GeometryMask/Private/GeometryMaskCanvas.cpp
--- a/Engine/Plugins/VirtualProduction/GeometryMask/Source/GeometryMask/Private/GeometryMaskCanvas.cpp
+++ b/Engine/Plugins/VirtualProduction/GeometryMask/Source/GeometryMask/Private/GeometryMaskCanvas.cpp
@@ -121,19 +121,29 @@ void UGeometryMaskCanvas::Free()
 
 UCanvasRenderTarget2D* UGeometryMaskCanvas::GetTexture() const
 {
-	if (CanvasResource)
-	{
-		return CanvasResource->GetRenderTargetTexture();
-	}
-
 	return nullptr;
 }
 
+UTextureRenderTarget2DArray* UGeometryMaskCanvas::GetRenderTarget() const
+{
+	return CanvasResource ? CanvasResource->GetRenderTarget() : nullptr;
+}
+
+int16 UGeometryMaskCanvas::GetRenderTargetSliceIndex() const
+{
+	return CanvasResource ? CanvasResource->GetRenderTargetSliceIndex() : -1;
+}
+
 const FGeometryMaskCanvasId& UGeometryMaskCanvas::GetCanvasId() const
 {
 	return CanvasId;
 }
 
+FName UGeometryMaskCanvas::GetCanvasName() const
+{
+	return CanvasName;
+}
+
 bool UGeometryMaskCanvas::IsBlurApplied() const
 {
 	return bApplyBlur;
@@ -230,6 +240,13 @@ void UGeometryMaskCanvas::RemoveInvalidWriters()
 
 void UGeometryMaskCanvas::OnDrawToCanvas(const FGeometryMaskDrawingContext& InDrawingContext, FCanvas* InCanvas)
 {
+	if (!ensure(InCanvas))
+	{
+		return;
+	}
+
+	TRACE_CPUPROFILER_EVENT_SCOPE(UGeometryMaskCanvas::OnDrawToCanvas);
+
 	// Sort so subtract happens after additive, etc.
 	SortWriters();
 
@@ -237,7 +254,7 @@ void UGeometryMaskCanvas::OnDrawToCanvas(const FGeometryMaskDrawingContext& InDr
 	{
 		if (Writer.IsValid())
 		{
-			TRACE_CPUPROFILER_EVENT_SCOPE(UGeometryMaskCanvas::OnDrawToCanvas);
+			TRACE_CPUPROFILER_EVENT_SCOPE(UGeometryMaskCanvas::OnDrawToCanvas_Writer);
 			DECLARE_SCOPE_CYCLE_COUNTER(TEXT("UGeometryMaskCanvas::OnDrawToCanvas"), STAT_GeometryMask_Update, STATGROUP_GeometryMask);
 			Writer->DrawToCanvas(InCanvas);	
 		}
@@ -254,11 +271,16 @@ void UGeometryMaskCanvas::UpdateRenderParameters()
 
 void UGeometryMaskCanvas::Initialize(const ULevel* InLevel, FName InCanvasName)
 {
-	CanvasName = InCanvasName;
-	CanvasId = FGeometryMaskCanvasId(InLevel, CanvasName);
+}
+
+void UGeometryMaskCanvas::Initialize(const FInitParams& InInitParams)
+{
+	CanvasName = InInitParams.CanvasName;
+	CanvasId = FGeometryMaskCanvasId(InInitParams.Level, CanvasName);
 
 	CanvasResource = NewObject<UGeometryMaskCanvasResource>(this);
 	CanvasResource->OnDrawToCanvas().AddUObject(this, &UGeometryMaskCanvas::OnDrawToCanvas);
+	CanvasResource->Initialize(InInitParams.RenderTarget, InInitParams.SliceIndex);
 	UpdateRenderParameters();
 }
 
diff --git a/Engine/Plugins/VirtualProduction/GeometryMask/Source/GeometryMask/Private/GeometryMaskCanvasActor.cpp b/Engine/Plugins/VirtualProduction/GeometryMask/Source/GeometryMask/Private/GeometryMaskCanvasActor.cpp
--- a/Engine/Plugins/VirtualProduction/GeometryMask/Source/GeometryMask/Private/GeometryMaskCanvasActor.cpp
+++ b/Engine/Plugins/VirtualProduction/GeometryMask/Source/GeometryMask/Private/GeometryMaskCanvasActor.cpp
@@ -24,9 +24,14 @@ AGeometryMaskCanvasActor::AGeometryMaskCanvasActor(const FObjectInitializer& Obj
 	Canvas = nullptr;
 }
 
-UCanvasRenderTarget2D* AGeometryMaskCanvasActor::GetTexture()
+UCanvasRenderTarget2D* AGeometryMaskCanvasActor::GetTexture() const
 {
-	return Canvas ? Canvas->GetTexture() : nullptr;
+	return nullptr;
+}
+
+UTextureRenderTarget2DArray* AGeometryMaskCanvasActor::GetRenderTarget() const
+{
+	return Canvas ? Canvas->GetRenderTarget() : nullptr;
 }
 
 void AGeometryMaskCanvasActor::BeginPlay()
diff --git a/Engine/Plugins/VirtualProduction/GeometryMask/Source/GeometryMask/Private/GeometryMaskCanvasResource.cpp b/Engine/Plugins/VirtualProduction/GeometryMask/Source/GeometryMask/Private/GeometryMaskCanvasResource.cpp
--- a/Engine/Plugins/VirtualProduction/GeometryMask/Source/GeometryMask/Private/GeometryMaskCanvasResource.cpp
+++ b/Engine/Plugins/VirtualProduction/GeometryMask/Source/GeometryMask/Private/GeometryMaskCanvasResource.cpp
@@ -6,7 +6,9 @@
 #include "Engine/Canvas.h"
 #include "Engine/CanvasRenderTarget2D.h"
 #include "Engine/Level.h"
+#include "Engine/TextureRenderTarget2DArray.h"
 #include "Engine/World.h"
+#include "GeometryMaskCanvasUtils.h"
 #include "GeometryMaskModule.h"
 #include "GeometryMaskSettings.h"
 #include "GeometryMaskTypes.h"
@@ -54,6 +56,12 @@ void UGeometryMaskCanvasResource::BeginDestroy()
 	Super::BeginDestroy();
 }
 
+void UGeometryMaskCanvasResource::Initialize(UTextureRenderTarget2DArray* InRenderTarget, int32 InSliceIndex)
+{
+	RenderTarget = InRenderTarget;
+	RenderTargetSliceIndex = InSliceIndex;
+}
+
 const EGeometryMaskColorChannel UGeometryMaskCanvasResource::GetNextAvailableColorChannel() const
 {
 	return EGeometryMaskColorChannel::None;
@@ -106,11 +114,11 @@ TArray<FGeometryMaskCanvasId> UGeometryMaskCanvasResource::GetDependentCanvasIds
 	return {};
 }
 
-void UGeometryMaskCanvasResource::UpdateViewportSize()
+bool UGeometryMaskCanvasResource::UpdateViewportSize()
 {
-	if (!IsValid(RenderTargetTexture) || IsUnreachable())
+	if (!RenderTarget)
 	{
-		return;
+		return false;
 	}
 	
 	const float SizeMultiplier = GetDefault<UGeometryMaskSettings>()->GetDefaultResolutionMultiplier();
@@ -126,13 +134,13 @@ void UGeometryMaskCanvasResource::UpdateViewportSize()
 		NewViewportSize = NewViewportSize.ComponentMax(DrawingContext.ViewportSize);
 		
 		NewScaledViewportSize = NewScaledViewportSize.ComponentMax(
-			(DrawingContext.ViewportSize + GetViewportPadding(DrawingContext))
+			(DrawingContext.ViewportSize + GetViewportPadding())
 			* SizeMultiplier);
 	}
 
 	if (NewViewportSize.Size() == 0)
 	{
-		return;
+		return false;
 	}
 	
 	MaxViewportSize = NewViewportSize;
@@ -156,14 +164,16 @@ void UGeometryMaskCanvasResource::UpdateViewportSize()
 		SizeX /= RatioY;
 	}
 
-	if (RenderTargetTexture->SizeX == SizeX
-		&& RenderTargetTexture->SizeY == SizeY)
+	if (RenderTarget->SizeX == SizeX && RenderTarget->SizeY == SizeY)
 	{
-		return;
+		return false;
 	}
 
 	// Update RT size to viewport size
-	RenderTargetTexture->ResizeTarget(SizeX, SizeY);
+	RenderTarget->SizeX = SizeX;
+	RenderTarget->SizeY = SizeY;
+	UE::GeometryMask::UpdateRenderTarget(RenderTarget);
+	return true;
 }
 
 void UGeometryMaskCanvasResource::UpdateRenderParameters(
@@ -235,9 +245,6 @@ void UGeometryMaskCanvasResource::UpdateRenderParameters(
 	bApplyBlur = BlurParameters.bPerChannelApplyBlur.CountSetBits(0) > 0; // if any channels have blur
 	bApplyDF = DFParameters.bPerChannelCalculateDF.CountSetBits(0) > 0; // if any channels have DF
 
-	// create render target
-	GetRenderTargetTexture();
-
 	// Effects may have changed viewport padding
     UpdateViewportSize();
 }
@@ -256,7 +263,17 @@ void UGeometryMaskCanvasResource::SetViewportSize(FGeometryMaskDrawingContext& I
 	}
 }
 
+const FIntPoint& UGeometryMaskCanvasResource::GetMaxViewportSize() const
+{
+	return MaxViewportSize;
+}
+
 int32 UGeometryMaskCanvasResource::GetViewportPadding(const FGeometryMaskDrawingContext& InDrawingContext) const
+{
+	return GetViewportPadding();
+}
+
+int32 UGeometryMaskCanvasResource::GetViewportPadding() const
 {
 	int32 MaxFeatherRadius = 0;
 	{
@@ -287,18 +304,17 @@ int32 UGeometryMaskCanvasResource::GetViewportPadding(const FGeometryMaskDrawing
 
 UCanvasRenderTarget2D* UGeometryMaskCanvasResource::GetRenderTargetTexture()
 {
-	if (!IsValid(RenderTargetTexture))
-	{
-		const FName ObjectName = MakeUniqueObjectName(this, UCanvasRenderTarget2D::StaticClass(), FName(FString::Printf(TEXT("GeometryMaskCanvasResource_RenderTarget"))));
-		RenderTargetTexture = NewObject<UCanvasRenderTarget2D>(this, ObjectName);
-		RenderTargetTexture->bForceLinearGamma = false;
-		RenderTargetTexture->bAutoGenerateMips = false;
-		RenderTargetTexture->RenderTargetFormat = RTF_R16f; // single channel only
-		RenderTargetTexture->SetSampleCount(GetSampleCount());
-		RenderTargetTexture->InitAutoFormat(MaxViewportSize.X, MaxViewportSize.Y);
-	}
+	return nullptr;
+}
 
-	return RenderTargetTexture;
+UTextureRenderTarget2DArray* UGeometryMaskCanvasResource::GetRenderTarget() const
+{
+	return RenderTarget;
+}
+
+int16 UGeometryMaskCanvasResource::GetRenderTargetSliceIndex() const
+{
+	return RenderTargetSliceIndex;
 }
 
 void UGeometryMaskCanvasResource::Update(const ULevel* InLevel, FSceneView& InView, int32 InViewIndex)
@@ -314,7 +330,7 @@ ETextureRenderTargetSampleCount UGeometryMaskCanvasResource::GetSampleCount() co
 
 void UGeometryMaskCanvasResource::Draw(const ULevel* InLevel, FSceneView& InView, int32 InViewIndex)
 {
-	if (!InLevel || !InLevel->OwningWorld)
+	if (!InLevel || !InLevel->OwningWorld || !RenderTarget || RenderTargetSliceIndex == -1 || !RenderTarget->GameThread_GetRenderTargetResource())
 	{
 		return;
 	}
@@ -323,79 +339,92 @@ void UGeometryMaskCanvasResource::Draw(const ULevel* InLevel, FSceneView& InView
 
 	UE_LOG(LogGeometryMask, VeryVerbose, TEXT("UGeometryMaskCanvasResource::Update Level: %s"), *InLevel->GetName());
 
-	if (UCanvasRenderTarget2D* Texture = GetRenderTargetTexture())
+	// First time initialization
+	if (!Canvas.IsValid())
 	{
-		// First time initialization
-		if (!Canvas.IsValid())
+		Canvas = MakeShared<FCanvas>(
+			RenderTarget->GameThread_GetRenderTargetResource(),
+			nullptr,
+			InLevel->OwningWorld,
+			InLevel->OwningWorld->GetFeatureLevel(),
+			// Draw immediately so that interleaved SetVectorParameter (etc) function calls work as expected
+			FCanvas::CDM_ImmediateDrawing);
+	}
+
+	// Update render target as it could've changed (e.g. when changing slice count)
+	Canvas->SetSliceIndex(RenderTargetSliceIndex);
+	Canvas->SetRenderTarget_GameThread(RenderTarget->GameThread_GetRenderTargetResource());
+
+	FGeometryMaskDrawingContext* const DrawingContext = GetDrawingContextForLevel(InLevel, InViewIndex);
+    check(DrawingContext); // Should always be valid - world was already checked
+
+	bool bViewportSizeUpdated = false;
+
+	// Begin
+	{
+		const FIntPoint ViewportSize = InView.UnconstrainedViewRect.Size();
+		if (ViewportSize.SizeSquared() > 0 && DrawingContext->ViewportSize != ViewportSize)
 		{
-			Canvas = MakeShared<FCanvas>(
-				Texture->GameThread_GetRenderTargetResource(),
-				nullptr,
-				InLevel->OwningWorld,
-				InLevel->OwningWorld->GetFeatureLevel(),
-				// Draw immediately so that interleaved SetVectorParameter (etc) function calls work as expected
-				FCanvas::CDM_ImmediateDrawing);
+			DrawingContext->ViewportSize = ViewportSize;
+			bViewportSizeUpdated = UpdateViewportSize();
 		}
 
-		FGeometryMaskDrawingContext* DrawingContextPtr = GetDrawingContextForLevel(InLevel, InViewIndex);
-        check(DrawingContextPtr); // Should always be valid - world was already checked
+		FMatrix ProjectionMatrix = InView.ViewMatrices.GetProjectionMatrix();
+		UE::GeometryMask::Private::OverscanProjectionMatrix(ProjectionMatrix, MaxViewportSize, GetViewportPadding());
 
-        FGeometryMaskDrawingContext& DrawingContextRef = *DrawingContextPtr;
+		DrawingContext->ViewProjectionMatrix = InView.ViewMatrices.GetViewMatrix() * ProjectionMatrix;
+		DrawingContext->ViewProjectionMatrix.M[2][2] = UE_KINDA_SMALL_NUMBER; // Prevents div by zero later
 
-		// Begin
+		// Update viewport size
+		if (bViewportSizeUpdated || UpdateViewportSize())
 		{
-			SetViewportSize(DrawingContextRef, InView.UnconstrainedViewRect.Size());
-			FMatrix ProjectionMatrix = InView.ViewMatrices.GetProjectionMatrix();
-			UE::GeometryMask::Private::OverscanProjectionMatrix(ProjectionMatrix, MaxViewportSize, GetViewportPadding(DrawingContextRef));
-
-			DrawingContextRef.ViewProjectionMatrix = InView.ViewMatrices.GetViewMatrix() * ProjectionMatrix;
-			DrawingContextRef.ViewProjectionMatrix.M[2][2] = UE_KINDA_SMALL_NUMBER; // Prevents div by zero later
+			FlushRenderingCommands();
+			Canvas->SetRenderTarget_GameThread(RenderTarget->GameThread_GetRenderTargetResource());
+		}
 
-			UpdateViewportSize();
+		InLevel->OwningWorld->FlushDeferredParameterCollectionInstanceUpdates();
 
-			InLevel->OwningWorld->FlushDeferredParameterCollectionInstanceUpdates();
+		Canvas->SetParentCanvasSize(FIntPoint(RenderTarget->SizeX, RenderTarget->SizeY));
+		Canvas->Clear(RenderTarget->ClearColor);
+		RenderTarget->UpdateResourceImmediate(/*Clear*/false);
+	}
 
-			Canvas->SetParentCanvasSize(FIntPoint(Texture->SizeX, Texture->SizeY));
-			Canvas->Clear(Texture->ClearColor);
-			Texture->UpdateResourceImmediate(/** Clear */true);
-		}
+	// Contents
+	{
+		// Store current transform
+		const FMatrix CanvasMatrix = Canvas->GetBottomTransform();
 
-		// Contents
+		// Set to World->Viewport transform
+		Canvas->SetBaseTransform(DrawingContext->ViewProjectionMatrix);
 		{
-			// Store current transform
-			const FMatrix CanvasMatrix = Canvas->GetBottomTransform();
+			TRACE_CPUPROFILER_EVENT_SCOPE(UGeometryMaskCanvasResource::Draw::OnDrawToCanvas);
+			OnDrawToCanvas().Broadcast(*DrawingContext, Canvas.Get());
+		}
 
-			// Set to World->Viewport transform
-			Canvas->SetBaseTransform(DrawingContextRef.ViewProjectionMatrix);
-			{
-				TRACE_CPUPROFILER_EVENT_SCOPE(UGeometryMaskCanvasResource::Draw::OnDrawToCanvas);
-				OnDrawToCanvas().Broadcast(DrawingContextRef, Canvas.Get());
-			}
+		// Restore original transform
+		Canvas->SetBaseTransform(CanvasMatrix);
+	}
 
-			// Restore original transform
-			Canvas->SetBaseTransform(CanvasMatrix);
+	// End
+	{
+		// Commit batch elements to the rendering thread
+		{
+			TRACE_CPUPROFILER_EVENT_SCOPE(UGeometryMaskCanvasResource::Draw::Flush);
+			Canvas->Flush_GameThread();
 		}
 
-		// End
+		// Post Process
+		if (bApplyDF || bApplyBlur)
 		{
+			FRenderTarget* CanvasRenderTarget = Canvas->GetRenderTarget();
+			if (bApplyDF)
 			{
-				TRACE_CPUPROFILER_EVENT_SCOPE(UGeometryMaskCanvasResource::Draw::Flush);
-				Canvas->Flush_GameThread();
+				PostProcess_DistanceField->Execute(CanvasRenderTarget);
 			}
 
-			// Post Process
-			if (bApplyDF || bApplyBlur)
+			if (bApplyBlur)
 			{
-				FRenderTarget* RenderTarget = Canvas->GetRenderTarget();
-				if (bApplyDF)
-				{
-					PostProcess_DistanceField->Execute(RenderTarget);	
-				}
-
-				if (bApplyBlur)
-				{
-					PostProcess_Blur->Execute(RenderTarget);
-				}
+				PostProcess_Blur->Execute(CanvasRenderTarget);
 			}
 		}
 	}
diff --git a/Engine/Plugins/VirtualProduction/GeometryMask/Source/GeometryMask/Private/GeometryMaskCanvasUtils.cpp b/Engine/Plugins/VirtualProduction/GeometryMask/Source/GeometryMask/Private/GeometryMaskCanvasUtils.cpp
new file mode 100644
--- /dev/null
+++ b/Engine/Plugins/VirtualProduction/GeometryMask/Source/GeometryMask/Private/GeometryMaskCanvasUtils.cpp
@@ -0,0 +1,23 @@
+// Copyright Epic Games, Inc. All Rights Reserved.
+
+#include "GeometryMaskCanvasUtils.h"
+#include "RenderingThread.h"
+#include "Engine/TextureRenderTarget2DArray.h"
+
+namespace UE::GeometryMask
+{
+
+void UpdateRenderTarget(TNotNull<UTextureRenderTarget2DArray*> InRenderTarget)
+{
+	InRenderTarget->UpdateResource();
+
+	// Flush RHI thread after creating texture render target to make sure that RHIUpdateTextureReference is executed before doing any rendering with it
+	// This makes sure that Value->TextureReference.TextureReferenceRHI->GetReferencedTexture() is valid so that FUniformExpressionSet::FillUniformBuffer properly uses the texture for rendering, instead of using a fallback texture
+	ENQUEUE_RENDER_COMMAND(FlushRHIThreadToUpdateTextureRenderTargetReference)(
+		[](FRHICommandListImmediate& RHICmdList)
+		{
+			RHICmdList.ImmediateFlush(EImmediateFlushType::FlushRHIThread);
+		});
+}
+
+} // UE::GeometryMask
diff --git a/Engine/Plugins/VirtualProduction/GeometryMask/Source/GeometryMask/Private/GeometryMaskCanvasUtils.h b/Engine/Plugins/VirtualProduction/GeometryMask/Source/GeometryMask/Private/GeometryMaskCanvasUtils.h
new file mode 100644
--- /dev/null
+++ b/Engine/Plugins/VirtualProduction/GeometryMask/Source/GeometryMask/Private/GeometryMaskCanvasUtils.h
@@ -0,0 +1,14 @@
+// Copyright Epic Games, Inc. All Rights Reserved.
+
+#pragma once
+
+#include "Misc/NotNull.h"
+
+class UTextureRenderTarget2DArray;
+
+namespace UE::GeometryMask
+{
+	/** Updates the resource for the render target */
+	void UpdateRenderTarget(TNotNull<UTextureRenderTarget2DArray*> InRenderTarget);
+
+} // UE::GeometryMask
diff --git a/Engine/Plugins/VirtualProduction/GeometryMask/Source/GeometryMask/Private/GeometryMaskSubsystem.cpp b/Engine/Plugins/VirtualProduction/GeometryMask/Source/GeometryMask/Private/GeometryMaskSubsystem.cpp
--- a/Engine/Plugins/VirtualProduction/GeometryMask/Source/GeometryMask/Private/GeometryMaskSubsystem.cpp
+++ b/Engine/Plugins/VirtualProduction/GeometryMask/Source/GeometryMask/Private/GeometryMaskSubsystem.cpp
@@ -17,14 +17,7 @@
 
 UGeometryMaskCanvas* UGeometryMaskSubsystem::GetDefaultCanvas()
 {
-	if (!DefaultCanvas)
-	{
-		const FName DefaultCanvasObjectName = FName(TEXT("GeometryMaskCanvas_Default"));
-		DefaultCanvas = NewObject<UGeometryMaskCanvas>(this, DefaultCanvasObjectName);
-		DefaultCanvas->Initialize(nullptr, FGeometryMaskCanvasId::DefaultCanvasName);
-	}
-	
-	return DefaultCanvas;
+	return nullptr;
 }
 
 int32 UGeometryMaskSubsystem::GetNumCanvasResources() const
@@ -86,7 +79,7 @@ void UGeometryMaskSubsystem::UpdateLevel(const ULevel* InLevel, UGeometryMaskWor
 	{
 		FSceneView* MutableSceneView = const_cast<FSceneView*>(View);
 
-		TArray<UGeometryMaskCanvasResource*, TInlineAllocator<1>> UsedResources;
+		TArray<UGeometryMaskCanvasResource*, TInlineAllocator<4>> UsedResources;
 
 		if (const FGeometryMaskLevelState* LevelState = InWorldSubsystem->FindLevelState(InLevel))
 		{
diff --git a/Engine/Plugins/VirtualProduction/GeometryMask/Source/GeometryMask/Private/GeometryMaskTypes.cpp b/Engine/Plugins/VirtualProduction/GeometryMask/Source/GeometryMask/Private/GeometryMaskTypes.cpp
--- a/Engine/Plugins/VirtualProduction/GeometryMask/Source/GeometryMask/Private/GeometryMaskTypes.cpp
+++ b/Engine/Plugins/VirtualProduction/GeometryMask/Source/GeometryMask/Private/GeometryMaskTypes.cpp
@@ -137,16 +137,17 @@ UGeometryMaskCanvasReferenceComponentBase::~UGeometryMaskCanvasReferenceComponen
 	Cleanup();
 }
 
-UCanvasRenderTarget2D* UGeometryMaskCanvasReferenceComponentBase::GetTexture()
+UCanvasRenderTarget2D* UGeometryMaskCanvasReferenceComponentBase::GetTexture() const
 {
-	if (const UGeometryMaskCanvas* Canvas = CanvasWeak.Get())
-	{
-		return Canvas->GetTexture();
-	}
-	
 	return nullptr;
 }
 
+UTextureRenderTarget2DArray* UGeometryMaskCanvasReferenceComponentBase::GetRenderTarget() const
+{
+	const UGeometryMaskCanvas* Canvas = CanvasWeak.Get();
+	return Canvas ? Canvas->GetRenderTarget() : nullptr;
+}
+
 void UGeometryMaskCanvasReferenceComponentBase::BeginPlay()
 {
 	Super::BeginPlay();
diff --git a/Engine/Plugins/VirtualProduction/GeometryMask/Source/GeometryMask/Private/GeometryMaskWorldSubsystem.cpp b/Engine/Plugins/VirtualProduction/GeometryMask/Source/GeometryMask/Private/GeometryMaskWorldSubsystem.cpp
--- a/Engine/Plugins/VirtualProduction/GeometryMask/Source/GeometryMask/Private/GeometryMaskWorldSubsystem.cpp
+++ b/Engine/Plugins/VirtualProduction/GeometryMask/Source/GeometryMask/Private/GeometryMaskWorldSubsystem.cpp
@@ -4,7 +4,9 @@
 
 #include "Engine/Engine.h"
 #include "Engine/Level.h"
+#include "Engine/TextureRenderTarget2DArray.h"
 #include "Engine/World.h"
+#include "GeometryMaskCanvasUtils.h"
 #include "GeometryMaskSVE.h"
 #include "SceneViewExtension.h"
 
@@ -38,35 +40,54 @@ void UGeometryMaskWorldSubsystem::Deinitialize()
 
 const FGeometryMaskLevelState* UGeometryMaskWorldSubsystem::FindLevelState(const ULevel* InLevel) const
 {
-	if (IsValid(InLevel))
-	{
-		return LevelStates.Find(InLevel);
-	}
-	return nullptr;
+	return InLevel ? LevelStates.Find(InLevel) : nullptr;
 }
 
 FGeometryMaskLevelState& UGeometryMaskWorldSubsystem::FindOrAddLevelState(const ULevel* InLevel)
 {
 	check(IsValid(InLevel));
-	return LevelStates.FindOrAdd(InLevel);
-}
 
-UGeometryMaskCanvas* UGeometryMaskWorldSubsystem::GetNamedCanvas(const ULevel* InLevel, FName InName)
-{
-	if (!IsValid(InLevel))
+	FGeometryMaskLevelState* LevelState = LevelStates.Find(InLevel);
+	if (!LevelState)
 	{
-		return nullptr;
+		LevelState = &LevelStates.Add(InLevel);
+
+		const FName ObjectName = MakeUniqueObjectName(this, UTextureRenderTarget2DArray::StaticClass(), FName(FString::Printf(TEXT("GeometryMaskRenderTarget"))));
+		LevelState->RenderTarget = NewObject<UTextureRenderTarget2DArray>(this, ObjectName);
+
+		LevelState->RenderTarget->ClearColor = FLinearColor::Black;
+		LevelState->RenderTarget->bForceLinearGamma = false;
+		LevelState->RenderTarget->bTargetArraySlicesIndependently = true;
+
+		constexpr int32 Width = 1920;
+		constexpr int32 Height = 1080;
+		constexpr int32 Slices = 3;
+		LevelState->RenderTarget->Init(Width, Height, Slices, PF_R16F);
 	}
+	return *LevelState;
+}
 
-	UGeometryMaskSubsystem* EngineSubsystem = GEngine->GetEngineSubsystem<UGeometryMaskSubsystem>();
-	if (!ensureAlwaysMsgf(EngineSubsystem, TEXT("UGeometryMaskSubsystem not resolved.")))
+int32 FGeometryMaskLevelState::AcquireNextSliceIndex()
+{
+	if (RenderTarget)
 	{
-		return nullptr;
+		const int32 SliceIndex = NextSliceIndex++;
+		if (SliceIndex >= RenderTarget->Slices)
+		{
+			// Update RT size to viewport size
+			RenderTarget->Slices = SliceIndex + 1;
+			UE::GeometryMask::UpdateRenderTarget(RenderTarget);
+		}
+		return SliceIndex;
 	}
+	return -1;
+}
 
-	if (InName.IsNone())
+UGeometryMaskCanvas* UGeometryMaskWorldSubsystem::GetNamedCanvas(const ULevel* InLevel, FName InName)
+{
+	if (InName.IsNone() || !InLevel)
 	{
-		return EngineSubsystem->GetDefaultCanvas();
+		return nullptr;
 	}
 
 	if (const FGeometryMaskLevelState* LevelState = FindLevelState(InLevel))
@@ -82,26 +103,18 @@ UGeometryMaskCanvas* UGeometryMaskWorldSubsystem::GetNamedCanvas(const ULevel* I
 	FGeometryMaskLevelState& LevelState = FindOrAddLevelState(InLevel);
 	const TObjectPtr<UGeometryMaskCanvas>& NewCanvas = LevelState.NamedCanvases.Emplace(InName, NewObject<UGeometryMaskCanvas>(this, ObjectName));
 
-	NewCanvas->Initialize(InLevel, InName);
+	UGeometryMaskCanvas::FInitParams InitParams;
+	InitParams.Level = InLevel;
+	InitParams.CanvasName = InName;
+	InitParams.RenderTarget = LevelState.RenderTarget;
+	InitParams.SliceIndex = LevelState.AcquireNextSliceIndex();
+
+	NewCanvas->Initialize(InitParams);
 	OnGeometryMaskCanvasCreatedDelegate.Broadcast(NewCanvas);
 
 	return NewCanvas;
 }
 
-void UGeometryMaskWorldSubsystem::ForEachCanvas(TFunctionRef<void(UGeometryMaskCanvas*, const FCanvasContext&)> InFunc) const
-{
-	FCanvasContext Context;
-	for (const TPair<TWeakObjectPtr<const ULevel>, FGeometryMaskLevelState>& LevelStatePair : LevelStates)
-	{
-		Context.LevelWeak = LevelStatePair.Key;
-		for (const TPair<FName, TObjectPtr<UGeometryMaskCanvas>>& CanvasPair : LevelStatePair.Value.NamedCanvases)
-		{
-			Context.CanvasName = CanvasPair.Key;
-			InFunc(CanvasPair.Value, Context);
-		}
-	}
-}
-
 TArray<FName> UGeometryMaskWorldSubsystem::GetCanvasNames(const ULevel* InLevel)
 {
 	TArray<FName> CanvasNames;
diff --git a/Engine/Plugins/VirtualProduction/GeometryMask/Source/GeometryMask/Public/GeometryMaskCanvas.h b/Engine/Plugins/VirtualProduction/GeometryMask/Source/GeometryMask/Public/GeometryMaskCanvas.h
--- a/Engine/Plugins/VirtualProduction/GeometryMask/Source/GeometryMask/Public/GeometryMaskCanvas.h
+++ b/Engine/Plugins/VirtualProduction/GeometryMask/Source/GeometryMask/Public/GeometryMaskCanvas.h
@@ -6,16 +6,14 @@
 #include "SceneView.h"
 #include "UObject/Object.h"
 #include "UObject/WeakInterfacePtr.h"
-
 #include "GeometryMaskCanvas.generated.h"
 
-class FGeometryMaskPostProcess_Blur;
+#define UE_API GEOMETRYMASK_API
+
 class IGeometryMaskWriteInterface;
-class UCanvas;
 class UCanvasRenderTarget2D;
 class UGeometryMaskCanvasResource;
-class UGeometryMaskPrimitive;
-class UGeometryMaskSubsystem;
+class UTextureRenderTarget2DArray;
 
 /** Called when Writers becomes non-empty. */
 using FOnGeometryMaskCanvasActivated = TDelegate<void()>;
@@ -24,93 +22,122 @@ using FOnGeometryMaskCanvasActivated = TDelegate<void()>;
 using FOnGeometryMaskCanvasDeactivated = TDelegate<void()>;
 
 /** A uniquely identified Canvas. */
-UCLASS(BlueprintType, Transient)
-class GEOMETRYMASK_API UGeometryMaskCanvas : public UObject
+UCLASS(MinimalAPI, BlueprintType, Transient)
+class UGeometryMaskCanvas : public UObject
 {
 	GENERATED_BODY()
 
 public:
 #if WITH_EDITOR
-	virtual void PostEditChangeProperty(FPropertyChangedEvent& InPropertyChangedEvent) override;
+	UE_API virtual void PostEditChangeProperty(FPropertyChangedEvent& InPropertyChangedEvent) override;
 #endif
 
-	ULevel* GetLevel() const;
+	UE_API ULevel* GetLevel() const;
 
 	/** Returns all writers. */
-    const TArray<TWeakInterfacePtr<IGeometryMaskWriteInterface>>& GetWriters() const;
+    UE_API const TArray<TWeakInterfacePtr<IGeometryMaskWriteInterface>>& GetWriters() const;
 
 	/** Adds a Mask Writer to the canvas. */
-	void AddWriter(const TScriptInterface<IGeometryMaskWriteInterface>& InWriter);
+	UE_API void AddWriter(const TScriptInterface<IGeometryMaskWriteInterface>& InWriter);
 
 	/** Appends multiple writers to the canvas. */
-	void AddWriters(const TArray<TScriptInterface<IGeometryMaskWriteInterface>>& InWriters);
+	UE_API void AddWriters(const TArray<TScriptInterface<IGeometryMaskWriteInterface>>& InWriters);
 
 	/** Remove a writer from this canvas. */
-	void RemoveWriter(const TScriptInterface<IGeometryMaskWriteInterface>& InWriter);
+	UE_API void RemoveWriter(const TScriptInterface<IGeometryMaskWriteInterface>& InWriter);
 
 	/** Gets the number of writers for the canvas. */
-	int32 GetNumWriters() const;
+	UE_API int32 GetNumWriters() const;
 
 	/** Returns true if this is the default/blank canvas. */
-	bool IsDefaultCanvas() const;
+	UE_API bool IsDefaultCanvas() const;
 
 	/** Forcibly free the canvas - removed all writers and frees the resource. */
-	void Free();
+	UE_API void Free();
 
-	/** Get the underlying render target. */
+	UE_DEPRECATED(5.8, "GetTexture (UCanvasRenderTarget2D*) has been deprecated. Use GetRenderTarget (UTextureRenderTarget2DArray*) instead")
+	UFUNCTION(BlueprintCallable, BlueprintPure, Category = "Rendering", meta=(DeprecatedFunction, DeprecatedMessage="GetTexture (CanvasRenderTarget2D) has been deprecated in favor of GetRenderTarget (TextureRenderTarget2DArray)"))
+	UE_API UCanvasRenderTarget2D* GetTexture() const;
+
+	/** Returns the render target 2D array that the canvas draws to */
 	UFUNCTION(BlueprintCallable, BlueprintPure, Category = "Rendering")
-	UCanvasRenderTarget2D* GetTexture() const;
+	UE_API UTextureRenderTarget2DArray* GetRenderTarget() const;
+
+	/** Get the index to the render target slice the canvas draws to */
+	UE_API int16 GetRenderTargetSliceIndex() const;
 
 	/** Get the unique canvas id based on the world and canvas name. */
-	const FGeometryMaskCanvasId& GetCanvasId() const;
+	UE_API const FGeometryMaskCanvasId& GetCanvasId() const;
+
+	/** Returns the canvas name */
+	UE_API FName GetCanvasName() const;
 
 	/** Whether blur is applied or not. */
-	bool IsBlurApplied() const;
+	UE_API bool IsBlurApplied() const;
 
 	/** Sets whether blur is applied or not. */
-	void SetApplyBlur(const bool bInValue);
+	UE_API void SetApplyBlur(const bool bInValue);
 
 	/** Gets the current blur strength, if applicable. */
-	double GetBlurStrength() const;
+	UE_API double GetBlurStrength() const;
 
 	/** Sets the blur strength. 0.0 will disable blur. */
-	void SetBlurStrength(const double InValue);
+	UE_API void SetBlurStrength(const double InValue);
 
 	/** Whether feathering is applied or not. */
-	bool IsFeatherApplied() const;
+	UE_API bool IsFeatherApplied() const;
 
 	/** Sets whether feathering is applied or not. */
-	void SetApplyFeather(const bool bInValue);
+	UE_API void SetApplyFeather(const bool bInValue);
 
 	/** Gets the current outer feather radius, in pixels. */
-	int32 GetOuterFeatherRadius() const;
+	UE_API int32 GetOuterFeatherRadius() const;
 
 	/** Sets the outer feather radius, in pixels. */
-	void SetOuterFeatherRadius(const int32 InValue);
+	UE_API void SetOuterFeatherRadius(const int32 InValue);
 
 	/** Gets the current inner feather radius, in pixels. */
-	int32 GetInnerFeatherRadius() const;
+	UE_API int32 GetInnerFeatherRadius() const;
 
 	/** Sets the inner feather radius, in pixels. */
-	void SetInnerFeatherRadius(const int32 InValue);
+	UE_API void SetInnerFeatherRadius(const int32 InValue);
 
 	/** Get the color channel to write to in the texture. */
 	UE_DEPRECATED(5.8, "Color channel management has been deprecated as it's no longer needed")
 	UFUNCTION(BlueprintCallable, BlueprintPure, Category = "Rendering", meta=(DeprecatedFunction, DeprecationMessage="Color channel management has been deprecated as it's no longer needed"))
-	EGeometryMaskColorChannel GetColorChannel() const;
+	UE_API EGeometryMaskColorChannel GetColorChannel() const;
 
 	/** Called when Writers becomes non-empty. */
-	FOnGeometryMaskCanvasActivated& OnActivated() { return OnActivatedDelegate; }
+	FOnGeometryMaskCanvasActivated& OnActivated()
+	{
+		return OnActivatedDelegate;
+	}
 
 	/** Called when Writers becomes empty. */
-	FOnGeometryMaskCanvasDeactivated& OnDeactivated() { return OnDeactivatedDelegate; }
+	FOnGeometryMaskCanvasDeactivated& OnDeactivated()
+	{
+		return OnDeactivatedDelegate;
+	}
+
+	UE_DEPRECATED(5.8, "Use the Initialize() overload with init params")
+	UE_API void Initialize(const ULevel* InLevel, FName InCanvasName);
+
+	struct FInitParams
+	{
+		/** Level that this canvas is relevant for */
+		const ULevel* Level = nullptr;
+		/** Name identifying this canvas within the level */
+		FName CanvasName;
+		/** Render target to draw to */
+		UTextureRenderTarget2DArray* RenderTarget = nullptr;
+		/** Index to the slice in the render target to draw to */
+		int32 SliceIndex = -1;
+	};
+	UE_API void Initialize(const FInitParams& InInitParams);
 
-	/** Setup with identifying info. */
-	void Initialize(const ULevel* InLevel, FName InCanvasName);
-	
 	/** Updates the canvas, intended to be called every frame. */
-	void Update(const ULevel* InLevel, FSceneView& InView);
-	
+	UE_API void Update(const ULevel* InLevel, FSceneView& InView);
+
 	const UGeometryMaskCanvasResource* GetResource() const
 	{
 		return CanvasResource;
@@ -122,17 +149,17 @@ public:
 	}
 
 	UE_DEPRECATED(5.8, "AssignResource has been deprecated as resource management is done by the canvas itself")
-	void AssignResource(UGeometryMaskCanvasResource* InResource, EGeometryMaskColorChannel InColorChannel);
+	UE_API void AssignResource(UGeometryMaskCanvasResource* InResource, EGeometryMaskColorChannel InColorChannel);
 
 	UE_DEPRECATED(5.8, "FreeResource has been deprecated as resource management is done by the canvas itself")
-	void FreeResource();
+	UE_API void FreeResource();
 
 public:
-	static FName GetApplyBlurPropertyName();
-	static FName GetBlurStrengthPropertyName();
-	static FName GetApplyFeatherPropertyName();
-	static FName GetOuterFeatherRadiusPropertyName();
-	static FName GetInnerFeatherRadiusPropertyName();
+	UE_API static FName GetApplyBlurPropertyName();
+	UE_API static FName GetBlurStrengthPropertyName();
+	UE_API static FName GetApplyFeatherPropertyName();
+	UE_API static FName GetOuterFeatherRadiusPropertyName();
+	UE_API static FName GetInnerFeatherRadiusPropertyName();
 
 private:
 	static const FName ApplyBlurPropertyName;
@@ -164,7 +191,7 @@ private:
 	FGeometryMaskCanvasId CanvasId;
 
 	/** Uniquely identifies this canvas. */
-	UPROPERTY(VisibleAnywhere, BlueprintReadOnly, Getter = "Auto", Category = "Canvas", meta = (AllowPrivateAccess = "true"))
+	UPROPERTY(VisibleAnywhere, BlueprintReadOnly, Getter, Category = "Canvas", meta = (AllowPrivateAccess = "true"))
 	FName CanvasName;
 
 	/** Optional Blur Toggle. */
@@ -194,3 +221,5 @@ private:
 	/** List of objects that write to this canvas. */
 	TArray<TWeakInterfacePtr<IGeometryMaskWriteInterface>> Writers;
 };
+
+#undef UE_API
diff --git a/Engine/Plugins/VirtualProduction/GeometryMask/Source/GeometryMask/Public/GeometryMaskCanvasActor.h b/Engine/Plugins/VirtualProduction/GeometryMask/Source/GeometryMask/Public/GeometryMaskCanvasActor.h
--- a/Engine/Plugins/VirtualProduction/GeometryMask/Source/GeometryMask/Public/GeometryMaskCanvasActor.h
+++ b/Engine/Plugins/VirtualProduction/GeometryMask/Source/GeometryMask/Public/GeometryMaskCanvasActor.h
@@ -6,10 +6,10 @@
 
 #include "GeometryMaskCanvasActor.generated.h"
 
-class UGeometryMaskCanvas;
 class IGeometryMaskWriteInterface;
-class UGeometryMaskCaptureComponent;
-class UGeometryMaskWriteMeshComponent;
+class UCanvasRenderTarget2D;
+class UGeometryMaskCanvas;
+class UTextureRenderTarget2DArray;
 
 /** Wraps a GeometryMaskCanvas, and discovers/registers writers. */
 UCLASS(BlueprintType)
@@ -25,9 +25,13 @@ public:
 	UPROPERTY(EditAnywhere, BlueprintReadOnly, Category = "Canvas")
 	FName CanvasName;
 
-	/** Returns the Canvas Texture. */
-	UFUNCTION(BlueprintCallable, Category = "Rendering")
-	UCanvasRenderTarget2D* GetTexture();
+	UE_DEPRECATED(5.8, "GetTexture (UCanvasRenderTarget2D*) has been deprecated. Use GetRenderTarget (UTextureRenderTarget2DArray*) instead")
+	UFUNCTION(BlueprintCallable, BlueprintPure, Category = "Rendering", meta=(DeprecatedFunction, DeprecatedMessage="GetTexture (CanvasRenderTarget2D) has been deprecated in favor of GetRenderTarget (TextureRenderTarget2DArray)"))
+	UCanvasRenderTarget2D* GetTexture() const;
+
+	/** Returns the render target 2D array that the canvas draws to */
+	UFUNCTION(BlueprintCallable, BlueprintPure, Category = "Rendering")
+	UTextureRenderTarget2DArray* GetRenderTarget() const;
 
 	virtual void BeginPlay() override;
 
diff --git a/Engine/Plugins/VirtualProduction/GeometryMask/Source/GeometryMask/Public/GeometryMaskCanvasResource.h b/Engine/Plugins/VirtualProduction/GeometryMask/Source/GeometryMask/Public/GeometryMaskCanvasResource.h
--- a/Engine/Plugins/VirtualProduction/GeometryMask/Source/GeometryMask/Public/GeometryMaskCanvasResource.h
+++ b/Engine/Plugins/VirtualProduction/GeometryMask/Source/GeometryMask/Public/GeometryMaskCanvasResource.h
@@ -8,23 +8,21 @@
 
 #include "GeometryMaskCanvasResource.generated.h"
 
+#define UE_API GEOMETRYMASK_API
+
+class FCanvas;
 class FGeometryMaskPostProcess_Blur;
 class FGeometryMaskPostProcess_DistanceField;
-class FGeometryMaskPostProcess_SobelTest;
-class FCanvas;
 class FSceneView;
-class IGeometryMaskPostProcess;
-class UCanvas;
 class UCanvasRenderTarget2D;
+class UTextureRenderTarget2DArray;
 enum class EGeometryMaskColorChannel : uint8;
 enum class ETextureRenderTargetSampleCount : uint8;
 
 using FOnGeometryMaskCanvasDraw = TMulticastDelegate<void(const FGeometryMaskDrawingContext&, FCanvas*)>;
 
-/**  */
-UCLASS()
-class GEOMETRYMASK_API UGeometryMaskCanvasResource
-	: public UObject
+UCLASS(MinimalAPI)
+class UGeometryMaskCanvasResource : public UObject
 {
 	GENERATED_BODY()
 
@@ -33,54 +31,77 @@ public:
 	static constexpr int32 MaxNumAvailableChannels = UE::GeometryMask::MaxNumChannels;
 	static constexpr int32 MaxTextureSize = UE::GeometryMask::MaxTextureSize;
 
-	UGeometryMaskCanvasResource();
+	UE_API UGeometryMaskCanvasResource();
+
+	UE_API virtual void BeginDestroy() override;
 
-	virtual void BeginDestroy() override;
+	void Initialize(UTextureRenderTarget2DArray* InRenderTarget, int32 InSliceIndex);
 
 	/** Will return the first available color channel without a canvas assigned. EGeometryMaskColorChannel::None if not available. */
 	UE_DEPRECATED(5.8, "Color channel management has been deprecated as it's no longer needed")
-	const EGeometryMaskColorChannel GetNextAvailableColorChannel() const;
+	UE_API const EGeometryMaskColorChannel GetNextAvailableColorChannel() const;
 
 	/** Requests usage of the given color channel for this resource. Will return true if successful. */
 	UE_DEPRECATED(5.8, "Color channel management has been deprecated as it's no longer needed")
-	bool Checkout(const EGeometryMaskColorChannel InColorChannel, const FGeometryMaskCanvasId& InRequestingCanvasId);
+	UE_API bool Checkout(const EGeometryMaskColorChannel InColorChannel, const FGeometryMaskCanvasId& InRequestingCanvasId);
 
 	/** Returns/frees the color channel associated with the given canvas name. Returns true if canvas name found. */
 	UE_DEPRECATED(5.8, "Color channel management has been deprecated as it's no longer needed")
-	bool Checkin(const FGeometryMaskCanvasId& InRequestingCanvasId);
+	UE_API bool Checkin(const FGeometryMaskCanvasId& InRequestingCanvasId);
 
 	/** Re-arranges used channels such that they are used sequentially. Returns number of unused channels. */
 	UE_DEPRECATED(5.8, "Color channel management has been deprecated as it's no longer needed")
-	int32 Compact();
+	UE_API int32 Compact();
 
 	UE_DEPRECATED(5.8, "Color channel management has been deprecated as it's no longer needed")
-	int32 GetNumChannelsUsed() const;
+	UE_API int32 GetNumChannelsUsed() const;
 
 	UE_DEPRECATED(5.8, "Color channel management has been deprecated as it's no longer needed")
-	bool IsAnyChannelUsed() const;
+	UE_API bool IsAnyChannelUsed() const;
 
 	/** Get the list of dependent canvases, for debug purposes. */
 	UE_DEPRECATED(5.8, "Color channel management has been deprecated as it's no longer needed")
-	TArray<FGeometryMaskCanvasId> GetDependentCanvasIds() const;
+	UE_API TArray<FGeometryMaskCanvasId> GetDependentCanvasIds() const;
 
-	void UpdateViewportSize();
+	/** 
+	 * Validates that the render target size is up-to-date with the viewport size and updating it if not
+	 * @return true if the render target size was updated
+	 */
+	UE_API bool UpdateViewportSize();
 
-	void SetViewportSize(FGeometryMaskDrawingContext& InDrawingContext, const FIntPoint& InViewportSize);
+	UE_DEPRECATED(5.8, "SetViewportSize deprecated as it's no longer needed")
+	UE_API void SetViewportSize(FGeometryMaskDrawingContext& InDrawingContext, const FIntPoint& InViewportSize);
+
+	/** Returns the maximum size that would cover all requested viewports */
+	UE_API const FIntPoint& GetMaxViewportSize() const;
+
+	UE_DEPRECATED(5.8, "GetViewportPadding requiring a drawing context has been deprecated. Use the overload without the drawing context")
+	UE_API int32 GetViewportPadding(const FGeometryMaskDrawingContext& InDrawingContext) const;
 
 	/** Returns required viewport padding, in pixels - determined by certain effects. */
-	int32 GetViewportPadding(const FGeometryMaskDrawingContext& InDrawingContext) const;
+	UE_API int32 GetViewportPadding() const;
+
+	UE_DEPRECATED(5.8, "GetRenderTargetTexture (UCanvasRenderTarget2D) has been deprecated in favor of GetRenderTarget (UTextureRenderTarget2DArray)")
+	UE_API UCanvasRenderTarget2D* GetRenderTargetTexture();
 
-	UCanvasRenderTarget2D* GetRenderTargetTexture();
+	/** Returns the render target the resource draws to */
+	UE_API UTextureRenderTarget2DArray* GetRenderTarget() const;
 
-	FOnGeometryMaskCanvasDraw& OnDrawToCanvas() { return OnDrawToCanvasDelegate; }
+	/** Get the index to the render target slice the resource draws to */
+	UE_API int16 GetRenderTargetSliceIndex() const;
 
-	void UpdateRenderParameters(EGeometryMaskColorChannel InColorChannel, bool bInApplyBlur, double InBlurStrength, bool bInApplyFeather, int32 InOuterFeatherRadius, int32 InInnerFeatherRadius);
+	FOnGeometryMaskCanvasDraw& OnDrawToCanvas()
+	{
+		return OnDrawToCanvasDelegate;
+	}
+
+	UE_API void UpdateRenderParameters(EGeometryMaskColorChannel InColorChannel, bool bInApplyBlur, double InBlurStrength, bool bInApplyFeather, int32 InOuterFeatherRadius, int32 InInnerFeatherRadius);
 
 	/** Resets the render parameters for the given channel. */
-	void ResetRenderParameters(EGeometryMaskColorChannel InColorChannel);
+	UE_API void ResetRenderParameters(EGeometryMaskColorChannel InColorChannel);
 
 	/** Updates the canvas, intended to be called every frame. */
-	void Update(const ULevel* InLevel, FSceneView& InView, int32 InViewIndex = 0);
+	UE_API void Update(const ULevel* InLevel, FSceneView& InView, int32 InViewIndex = 0);
 
 private:
 	ETextureRenderTargetSampleCount GetSampleCount() const;
@@ -111,12 +132,21 @@ private:
 	TSharedPtr<FCanvas> Canvas;
 
 	/** The default viewport size to use when it can't be resolved from the actual viewport. */
-	UPROPERTY(EditAnywhere, BlueprintReadWrite, Getter = "Auto", Category = "Rendering", meta = (AllowPrivateAccess = "true"))
+	UPROPERTY(EditAnywhere, BlueprintReadWrite, Getter = "GetMaxViewportSize", Category = "Rendering", meta = (AllowPrivateAccess = "true"))
 	FIntPoint MaxViewportSize = FIntPoint(1920, 1080);
 
 	/** The underlying Render Target texture. */
+	UE_DEPRECATED(5.8, "RenderTargetTexture (UCanvasRenderTarget2D) has been deprecated in favor of RenderTarget (UTextureRenderTarget2DArray)")
+	UPROPERTY()
+	TObjectPtr<UCanvasRenderTarget2D> RenderTargetTexture_DEPRECATED;
+
+	/** The underlying Render Target. */
 	UPROPERTY(DuplicateTransient)
-	TObjectPtr<UCanvasRenderTarget2D> RenderTargetTexture;
+	TObjectPtr<UTextureRenderTarget2DArray> RenderTarget;
+
+	/** The render target slice assigned to this resource */
+	UPROPERTY(DuplicateTransient)
+	int16 RenderTargetSliceIndex = -1;
 
 	FOnGeometryMaskCanvasDraw OnDrawToCanvasDelegate;
 
@@ -130,3 +160,5 @@ private:
 	TSharedPtr<FGeometryMaskPostProcess_Blur> PostProcess_Blur;
 	TSharedPtr<FGeometryMaskPostProcess_DistanceField> PostProcess_DistanceField;
 };
+
+#undef UE_API
diff --git a/Engine/Plugins/VirtualProduction/GeometryMask/Source/GeometryMask/Public/GeometryMaskSubsystem.h b/Engine/Plugins/VirtualProduction/GeometryMask/Source/GeometryMask/Public/GeometryMaskSubsystem.h
--- a/Engine/Plugins/VirtualProduction/GeometryMask/Source/GeometryMask/Public/GeometryMaskSubsystem.h
+++ b/Engine/Plugins/VirtualProduction/GeometryMask/Source/GeometryMask/Public/GeometryMaskSubsystem.h
@@ -25,8 +25,8 @@ class UGeometryMaskSubsystem : public UEngineSubsystem
 	GENERATED_BODY()
 
 public:
-	/** Returns the default, blank canvas. */
-	UFUNCTION(BlueprintCallable, Category = "Canvas")
+	UE_DEPRECATED(5.8, "GetDefaultCanvas is deprecated. Subsystem no longer manages holds a default canvas")
+	UFUNCTION(BlueprintCallable, Category = "Canvas", meta=(DeprecatedFunction, DeprecatedMessage="GetDefaultCanvas is deprecated. Subsystem no longer manages holds a default canvas"))
 	UE_API UGeometryMaskCanvas* GetDefaultCanvas();
 
 	UE_DEPRECATED(5.8, "CanvasResources is deprecated. Subsystem no longer manages resources")
@@ -82,9 +82,6 @@ private:
 
 	std::atomic<bool> bDoUpdates = true;
 
-	UPROPERTY()
-	TObjectPtr<UGeometryMaskCanvas> DefaultCanvas;
-
 	/** Pool of GPU/Texture resources used by the canvases. */
 	UE_DEPRECATED(5.8, "CanvasResources is deprecated. Subsystem no longer manages resources")
 	UPROPERTY()
diff --git a/Engine/Plugins/VirtualProduction/GeometryMask/Source/GeometryMask/Public/GeometryMaskTypes.h b/Engine/Plugins/VirtualProduction/GeometryMask/Source/GeometryMask/Public/GeometryMaskTypes.h
--- a/Engine/Plugins/VirtualProduction/GeometryMask/Source/GeometryMask/Public/GeometryMaskTypes.h
+++ b/Engine/Plugins/VirtualProduction/GeometryMask/Source/GeometryMask/Public/GeometryMaskTypes.h
@@ -10,6 +10,7 @@
 
 class UCanvasRenderTarget2D;
 class UGeometryMaskCanvas;
+class UTextureRenderTarget2DArray;
 
 UENUM(BlueprintType)
 enum class EGeometryMaskColorChannel : uint8
@@ -247,9 +248,13 @@ public:
 	UFUNCTION(BlueprintImplementableEvent, CallInEditor, Category = "Canvas", meta = (DisplayName = "Canvas Changed"))
 	void ReceiveSetCanvas(const UGeometryMaskCanvas* InCanvas);
 
-	/** Returns the Canvas Texture. */
-	UFUNCTION(BlueprintCallable, Category="Rendering")
-	UCanvasRenderTarget2D* GetTexture();
+	UE_DEPRECATED(5.8, "GetTexture (UCanvasRenderTarget2D*) has been deprecated. Use GetRenderTarget (UTextureRenderTarget2DArray*) instead")
+	UFUNCTION(BlueprintCallable, BlueprintPure, Category = "Rendering", meta=(DeprecatedFunction, DeprecatedMessage="GetTexture (CanvasRenderTarget2D) has been deprecated in favor of GetRenderTarget (TextureRenderTarget2DArray)"))
+	UCanvasRenderTarget2D* GetTexture() const;
+
+	/** Returns the render target 2D array that the canvas draws to */
+	UFUNCTION(BlueprintCallable, BlueprintPure, Category = "Rendering")
+	UTextureRenderTarget2DArray* GetRenderTarget() const;
 
 protected:
 	virtual void BeginPlay() override;
diff --git a/Engine/Plugins/VirtualProduction/GeometryMask/Source/GeometryMask/Public/GeometryMaskWorldSubsystem.h b/Engine/Plugins/VirtualProduction/GeometryMask/Source/GeometryMask/Public/GeometryMaskWorldSubsystem.h
--- a/Engine/Plugins/VirtualProduction/GeometryMask/Source/GeometryMask/Public/GeometryMaskWorldSubsystem.h
+++ b/Engine/Plugins/VirtualProduction/GeometryMask/Source/GeometryMask/Public/GeometryMaskWorldSubsystem.h
@@ -8,6 +8,8 @@
 
 #include "GeometryMaskWorldSubsystem.generated.h"
 
+#define UE_API GEOMETRYMASK_API
+
 class FGeometryMaskSceneViewExtension;
 class ULevel;
 
@@ -19,51 +21,61 @@ struct FGeometryMaskLevelState
 {
 	GENERATED_BODY()
 
+	/** All the canvases in the level */
 	UPROPERTY()
 	TMap<FName, TObjectPtr<UGeometryMaskCanvas>> NamedCanvases;
+
+	/** The render target where the canvases in the levels draw to */
+	UPROPERTY()
+	TObjectPtr<UTextureRenderTarget2DArray> RenderTarget;
+
+	/** Returns the index to the next available slice in the render target */
+	int32 AcquireNextSliceIndex();
+
+private:
+	/** tracker for what the next available slice index is */
+	UPROPERTY()
+	int32 NextSliceIndex = 0;
 };
 
 /** Updates the canvases. */
-UCLASS()
-class GEOMETRYMASK_API UGeometryMaskWorldSubsystem
-	: public UWorldSubsystem
+UCLASS(MinimalAPI)
+class UGeometryMaskWorldSubsystem : public UWorldSubsystem
 {
 	GENERATED_BODY()
 
 public:
 	/** Retrieves a Canvas, uniquely identified by this world and the canvas name. */
-	UGeometryMaskCanvas* GetNamedCanvas(const ULevel* InLevel, FName InName);
-
-	struct FCanvasContext
-	{
-		FName CanvasName;
-		TWeakObjectPtr<const ULevel> LevelWeak;
-	};
-	/** Invokes the given function ref for each existing canvas */
-	void ForEachCanvas(TFunctionRef<void(UGeometryMaskCanvas*, const FCanvasContext&)> InFunc) const;
+	UE_API UGeometryMaskCanvas* GetNamedCanvas(const ULevel* InLevel, FName InName);
 
 	/** Returns all registered canvas names for this world. */
-	TArray<FName> GetCanvasNames(const ULevel* InLevel);
+	UE_API TArray<FName> GetCanvasNames(const ULevel* InLevel);
 	
 	/** Remove all canvases without any Readers or Writers. Return the number of canvases removed. */
-	int32 RemoveWithoutWriters();
+	UE_API int32 RemoveWithoutWriters();
 	
 	/** Called when a new canvas is created due to a unique name being requested. */
-	FOnGeometryMaskCanvasCreated& OnGeometryMaskCanvasCreated() { return OnGeometryMaskCanvasCreatedDelegate; }
+	FOnGeometryMaskCanvasCreated& OnGeometryMaskCanvasCreated()
+	{
+		return OnGeometryMaskCanvasCreatedDelegate;
+	}
 
 	/** Called when a canvas is destroyed due to having no registered writers. */
-	FOnGeometryMaskCanvasDestroyed& OnGeometryMaskCanvasDestroyed() { return OnGeometryMaskCanvasDestroyedDelegate; }
+	FOnGeometryMaskCanvasDestroyed& OnGeometryMaskCanvasDestroyed()
+	{
+		return OnGeometryMaskCanvasDestroyedDelegate;
+	}
 
 protected:
 	// ~Begin USubsystem
-	virtual void Initialize(FSubsystemCollectionBase& Collection) override;
-	virtual void Deinitialize() override;
+	UE_API virtual void Initialize(FSubsystemCollectionBase& Collection) override;
+	UE_API virtual void Deinitialize() override;
 	// ~End USubsystem
 
 private:
-	const FGeometryMaskLevelState* FindLevelState(const ULevel* InLevel) const;
+	UE_API const FGeometryMaskLevelState* FindLevelState(const ULevel* InLevel) const;
 
-	FGeometryMaskLevelState& FindOrAddLevelState(const ULevel* InLevel);
+	UE_API FGeometryMaskLevelState& FindOrAddLevelState(const ULevel* InLevel);
 
 	friend class UGeometryMaskSubsystem;
 
@@ -75,3 +87,5 @@ private:
 	UPROPERTY()
 	TMap<TWeakObjectPtr<const ULevel>, FGeometryMaskLevelState> LevelStates;
 };
+
+#undef UE_API
diff --git a/Engine/Plugins/VirtualProduction/GeometryMask/Source/GeometryMaskEditor/Private/ViewModels/GMECanvasItemViewModel.cpp b/Engine/Plugins/VirtualProduction/GeometryMask/Source/GeometryMaskEditor/Private/ViewModels/GMECanvasItemViewModel.cpp
--- a/Engine/Plugins/VirtualProduction/GeometryMask/Source/GeometryMaskEditor/Private/ViewModels/GMECanvasItemViewModel.cpp
+++ b/Engine/Plugins/VirtualProduction/GeometryMask/Source/GeometryMaskEditor/Private/ViewModels/GMECanvasItemViewModel.cpp
@@ -2,9 +2,9 @@
 
 #include "GMECanvasItemViewModel.h"
 
-#include "Engine/CanvasRenderTarget2D.h"
 #include "Engine/Level.h"
 #include "Engine/Texture.h"
+#include "Engine/TextureRenderTarget2DArray.h"
 #include "GeometryMaskCanvas.h"
 
 TSharedRef<FGMECanvasItemViewModel> FGMECanvasItemViewModel::Create(const TWeakObjectPtr<const UGeometryMaskCanvas>& InCanvas)
@@ -24,7 +24,7 @@ FGMECanvasItemViewModel::FGMECanvasItemViewModel(
 	if (const UGeometryMaskCanvas* Canvas = InCanvas.Get())
 	{
 		CanvasId = Canvas->GetCanvasId();		
-		CanvasTextureWeak = Canvas->GetTexture();
+		CanvasTextureWeak = Canvas->GetRenderTarget();
 		KnownWriterCount = Canvas->GetWriters().Num();
 
 		UpdateInfoText();
diff --git a/Engine/Plugins/VirtualProduction/GeometryMask/Source/GeometryMaskEditor/Private/Widgets/SGeometryMaskCanvasPreview.cpp b/Engine/Plugins/VirtualProduction/GeometryMask/Source/GeometryMaskEditor/Private/Widgets/SGeometryMaskCanvasPreview.cpp
--- a/Engine/Plugins/VirtualProduction/GeometryMask/Source/GeometryMaskEditor/Private/Widgets/SGeometryMaskCanvasPreview.cpp
+++ b/Engine/Plugins/VirtualProduction/GeometryMask/Source/GeometryMaskEditor/Private/Widgets/SGeometryMaskCanvasPreview.cpp
@@ -6,7 +6,7 @@
 #include "Engine/CanvasRenderTarget2D.h"
 #include "Engine/Engine.h"
 #include "Engine/Level.h"
-#include "Engine/Texture.h"
+#include "Engine/TextureRenderTarget2DArray.h"
 #include "GeometryMaskCanvasResource.h"
 #include "GeometryMaskEditorLog.h"
 #include "GeometryMaskSubsystem.h"
@@ -24,13 +24,14 @@
 namespace UE::GeometryMaskEditor::Private
 {
 	constexpr const TCHAR* PreviewMaterialPath = TEXT("/GeometryMask/GeometryMask/M_GeometryMaskPreview");
-	constexpr const TCHAR* DefaultTexturePath = TEXT("/Engine/EngineResources/Black");
-
-	static FName TextureParameterName = TEXT("Mask_Texture");
-	static FName InvertParameterName = TEXT("Mask_Invert");
-	static FName PaddingParameterName = TEXT("Mask_Padding");
-	static FName OpacityParameterName = TEXT("OpacityMultiplier");
-	static FName FeatherParameterName = TEXT("Feather");
+	constexpr const TCHAR* DefaultTexturePath = TEXT("/GeometryMask/GeometryMask/GeometryMaskDefaultRenderTarget");
+
+	static FLazyName MaskTexturesParameterName = TEXT("Mask_Textures");
+	static FLazyName MaskIndicesParameterName = TEXT("Mask_TextureIndexVector");
+	static FLazyName InvertParameterName = TEXT("Mask_Invert");
+	static FLazyName PaddingParameterName = TEXT("Mask_Padding");
+	static FLazyName OpacityParameterName = TEXT("OpacityMultiplier");
+	static FLazyName FeatherParameterName = TEXT("Feather");
 }
 
 SGeometryMaskCanvasPreview::SGeometryMaskCanvasPreview()
@@ -74,14 +75,14 @@ void SGeometryMaskCanvasPreview::Construct(const FArguments& InArgs)
 
 	// Load DefaultTexture
 	{
-		if (DefaultTexture = Cast<UTexture>(DefaultTexturePath.TryLoad());
-			!DefaultTexture || !IsValid(DefaultTexture))
+		DefaultTexture = Cast<UTextureRenderTarget2DArray>(DefaultTexturePath.TryLoad());
+		if (!DefaultTexture)
 		{
 			UE_LOG(LogGeometryMaskEditor, Warning, TEXT("DefaultTexture could not be loaded from path: '%s'"), *DefaultTexturePath.ToString());
 		}
 	}
 
-	UpdateBrush(nullptr, DefaultTexture);
+	UpdateBrush(nullptr);
     TryResolveCanvas();
 
 	ChildSlot
@@ -114,8 +115,7 @@ void SGeometryMaskCanvasPreview::Construct(const FArguments& InArgs)
 void SGeometryMaskCanvasPreview::Tick(const FGeometry& InAllottedGeometry, const double InCurrentTime, const float InDeltaTime)
 {
 	SCompoundWidget::Tick(InAllottedGeometry, InCurrentTime, InDeltaTime);
-
-	UpdateBrush(CanvasWeak.Get(), nullptr);
+	UpdateBrush(CanvasWeak.Get());
 }
 
 const FGeometryMaskCanvasId& SGeometryMaskCanvasPreview::GetCanvasId() const
@@ -240,8 +240,9 @@ bool SGeometryMaskCanvasPreview::TryResolveCanvas()
 		{
 			if (UGeometryMaskWorldSubsystem* Subsystem = CanvasLevel->OwningWorld->GetSubsystem<UGeometryMaskWorldSubsystem>())
 			{
-				CanvasWeak = Subsystem->GetNamedCanvas(CanvasLevel, GetCanvasName());
-				UpdateBrush(CanvasWeak.Get(), nullptr);
+				UGeometryMaskCanvas* Canvas = Subsystem->GetNamedCanvas(CanvasLevel, GetCanvasName());
+				CanvasWeak = Canvas;
+				UpdateBrush(Canvas);
 			}
 		}
 	}
@@ -249,30 +250,26 @@ bool SGeometryMaskCanvasPreview::TryResolveCanvas()
 	return CanvasWeak != nullptr;
 }
 
-void SGeometryMaskCanvasPreview::UpdateBrush(const UGeometryMaskCanvas* InCanvas, UTexture* InTexture)
+void SGeometryMaskCanvasPreview::UpdateBrush(const UGeometryMaskCanvas* InCanvas)
 {
-	// If InTexture isn't specified, use Default/blank
-	if (!InTexture)
-	{
-		InTexture = DefaultTexture;
-	}
+	UTextureRenderTarget2DArray* Texture = DefaultTexture;
+	FLinearColor MaskIndices = FLinearColor(-1, -1, -1, -1);
 
 	const bool bHasValidMID = PreviewMID != nullptr;
-		
+
 	// When a canvas is provided, update the parameters from it, vs just from this widget
 	if (InCanvas)
 	{
-		if (UTexture* CanvasTexture = InCanvas->GetTexture())
+		if (UTextureRenderTarget2DArray* CanvasTexture = InCanvas->GetRenderTarget())
 		{
-			InTexture = CanvasTexture;
+			Texture = CanvasTexture;
+			MaskIndices.R = static_cast<float>(InCanvas->GetRenderTargetSliceIndex());
 		}
 
-		FGeometryMaskDrawingContext DrawingContext(InCanvas->GetCanvasId().Level);
-
 		FVector4f Padding(ForceInitToZero);
 		if (const UGeometryMaskCanvasResource* CanvasResource = InCanvas->GetResource())
 		{
-			FIntPoint ViewportPadding(CanvasResource->GetViewportPadding(DrawingContext));
+			FIntPoint ViewportPadding(CanvasResource->GetViewportPadding());
 			FVector2f ViewportPaddingF = FVector2f(ViewportPadding) / CanvasResource->GetMaxViewportSize();
 			Padding.X = ViewportPaddingF.X;
 			Padding.Y = ViewportPaddingF.Y;
@@ -286,22 +283,23 @@ void SGeometryMaskCanvasPreview::UpdateBrush(const UGeometryMaskCanvas* InCanvas
 			PreviewMID->SetVectorParameterValue(UE::GeometryMaskEditor::Private::PaddingParameterName, Padding);
 		}
 	}
-	
+
 	if (bHasValidMID)
 	{
-		PreviewMID->SetTextureParameterValue(UE::GeometryMaskEditor::Private::TextureParameterName, InTexture);
+		PreviewMID->SetTextureParameterValue(UE::GeometryMaskEditor::Private::MaskTexturesParameterName, Texture);
+		PreviewMID->SetVectorParameterValue(UE::GeometryMaskEditor::Private::MaskIndicesParameterName, MaskIndices);
 		PreviewMID->SetScalarParameterValue(UE::GeometryMaskEditor::Private::InvertParameterName, IsInverted() ? 1.0f : 0.0f);
 	}
 
 	if (PreviewBrush.IsValid())
 	{
 		FVector2D ImageSize(1920.f, 1080.f);
-		if (InTexture)
+		if (Texture)
 		{
-			ImageSize.X = InTexture->GetSurfaceWidth();
-			ImageSize.Y = InTexture->GetSurfaceHeight();
+			ImageSize.X = Texture->GetSurfaceWidth();
+			ImageSize.Y = Texture->GetSurfaceHeight();
 		}
-		
+
 		PreviewBrush->ImageSize = ImageSize;
 		AspectRatio.Set(ImageSize.X / ImageSize.Y);
 	}
@@ -312,6 +310,6 @@ void SGeometryMaskCanvasPreview::OnWorldCleanup(UWorld* InWorld, bool bInSession
 	const ULevel* const CanvasLevel = CanvasId.Get().Level.ResolveObjectPtr();
 	if (!CanvasLevel || CanvasLevel->GetWorld() == InWorld || CanvasLevel->GetOuter() == InWorld)
 	{
-		UpdateBrush(/*Canvas*/nullptr, /*Texture*/nullptr);
+		UpdateBrush(/*Canvas*/nullptr);
 	}
 }
diff --git a/Engine/Plugins/VirtualProduction/GeometryMask/Source/GeometryMaskEditor/Public/Widgets/SGeometryMaskCanvasPreview.h b/Engine/Plugins/VirtualProduction/GeometryMask/Source/GeometryMaskEditor/Public/Widgets/SGeometryMaskCanvasPreview.h
--- a/Engine/Plugins/VirtualProduction/GeometryMask/Source/GeometryMaskEditor/Public/Widgets/SGeometryMaskCanvasPreview.h
+++ b/Engine/Plugins/VirtualProduction/GeometryMask/Source/GeometryMaskEditor/Public/Widgets/SGeometryMaskCanvasPreview.h
@@ -9,7 +9,7 @@
 
 class UGeometryMaskCanvas;
 class UMaterialInstanceDynamic;
-class UTexture;
+class UTextureRenderTarget2DArray;
 
 /** Displays a named GeometryMaskCanvas. */
 class GEOMETRYMASKEDITOR_API SGeometryMaskCanvasPreview
@@ -85,7 +85,7 @@ public:
 private:
 	bool TryResolveCanvas();
 
-	void UpdateBrush(const UGeometryMaskCanvas* InCanvas, UTexture* InTexture);
+	void UpdateBrush(const UGeometryMaskCanvas* InCanvas);
 
 	/**
 	 * Called when the world is being cleaned up
@@ -99,7 +99,7 @@ private:
 	TObjectPtr<UMaterialInstanceDynamic> PreviewMID;
 
 	FSoftObjectPath DefaultTexturePath;
-	TObjectPtr<UTexture> DefaultTexture;
+	TObjectPtr<UTextureRenderTarget2DArray> DefaultTexture;
 
 	TWeakObjectPtr<UGeometryMaskCanvas> CanvasWeak;
 
