Motion Design: changed it so that multi masking is done via specifying the channels explicitly rather than implicitly through attachment hierarchy.
This is primarily done for flexibility and clarity.
Additional fixes to blur and feathering fallout from the mask refactor.

#Synced-CL 49564377

diff --git a/Engine/Plugins/VirtualProduction/Avalanche/Source/AvalancheMask/Private/Mask2D/AvaMask2DBaseModifier.cpp b/Engine/Plugins/VirtualProduction/Avalanche/Source/AvalancheMask/Private/Mask2D/AvaMask2DBaseModifier.cpp
--- a/Engine/Plugins/VirtualProduction/Avalanche/Source/AvalancheMask/Private/Mask2D/AvaMask2DBaseModifier.cpp
+++ b/Engine/Plugins/VirtualProduction/Avalanche/Source/AvalancheMask/Private/Mask2D/AvaMask2DBaseModifier.cpp
@@ -526,6 +526,27 @@ bool UAvaMask2DBaseModifier::TryResolveParentChannel()
 	return bParentChannelWasFound;
 }
 
+UGeometryMaskWorldSubsystem* UAvaMask2DBaseModifier::GetGeometryMaskWorldSubsystem(const ULevel** OutLevel) const
+{
+	const AActor* Actor = GetModifiedActor();
+	if (!Actor)
+	{
+		return nullptr;
+	}
+
+	const ULevel* Level = Actor->GetLevel();
+	if (!Level || !Level->OwningWorld)
+	{
+		return nullptr;
+	}
+
+	if (OutLevel)
+	{
+		*OutLevel = Level;
+	}
+	return Level->OwningWorld->GetSubsystem<UGeometryMaskWorldSubsystem>();
+}
+
 void UAvaMask2DBaseModifier::SetupMaskComponent(UActorComponent* InComponent)
 {
 }
@@ -599,19 +620,8 @@ void UAvaMask2DBaseModifier::OnMaskSetCanvas(const UGeometryMaskCanvas* InCanvas
 
 void UAvaMask2DBaseModifier::TryResolveCanvas()
 {
-	const AActor* Actor = GetModifiedActor();
-	if (!Actor)
-	{
-		return;
-	}
-
-	const ULevel* Level = Actor->GetLevel();
-	if (!Level || !Level->OwningWorld)
-	{
-		return;
-	}
-
-	if (UGeometryMaskWorldSubsystem* MaskSubsystem = Level->OwningWorld->GetSubsystem<UGeometryMaskWorldSubsystem>())
+	const ULevel* Level = nullptr;
+	if (UGeometryMaskWorldSubsystem* MaskSubsystem = GetGeometryMaskWorldSubsystem(&Level))
 	{
 		if (UGeometryMaskCanvas* Canvas = MaskSubsystem->GetNamedCanvas(Level, GetChannel()))
 		{
diff --git a/Engine/Plugins/VirtualProduction/Avalanche/Source/AvalancheMask/Private/Mask2D/AvaMask2DModifierShared.cpp b/Engine/Plugins/VirtualProduction/Avalanche/Source/AvalancheMask/Private/Mask2D/AvaMask2DModifierShared.cpp
deleted file mode 100644
diff --git a/Engine/Plugins/VirtualProduction/Avalanche/Source/AvalancheMask/Private/Mask2D/AvaMask2DReadModifier.cpp b/Engine/Plugins/VirtualProduction/Avalanche/Source/AvalancheMask/Private/Mask2D/AvaMask2DReadModifier.cpp
--- a/Engine/Plugins/VirtualProduction/Avalanche/Source/AvalancheMask/Private/Mask2D/AvaMask2DReadModifier.cpp
+++ b/Engine/Plugins/VirtualProduction/Avalanche/Source/AvalancheMask/Private/Mask2D/AvaMask2DReadModifier.cpp
@@ -10,7 +10,9 @@
 #include "Engine/Engine.h"
 #include "GeometryMaskCanvas.h"
 #include "GeometryMaskCanvasResource.h"
-#include "Mask2D/AvaMask2DModifierShared.h"
+#include "GeometryMaskCanvasSharedData.h"
+#include "GeometryMaskSubsystem.h"
+#include "GeometryMaskWorldSubsystem.h"
 #include "MaterialBridge/AvaMaterialBridge.h"
 #include "MaterialBridge/AvaMaterialBridgeRegistry.h"
 #include "MaterialBridge/Context/AvaMaterialBridgeApplyStateContext.h"
@@ -35,7 +37,7 @@
 const TAvaPropertyChangeDispatcher<UAvaMask2DReadModifier> UAvaMask2DReadModifier::PropertyChangeDispatcher =
 {
 	{ GET_MEMBER_NAME_CHECKED(UAvaMask2DReadModifier, BaseOpacity), &UAvaMask2DReadModifier::OnBaseOpacityChanged },
-	{ GET_MEMBER_NAME_CHECKED(UAvaMask2DReadModifier, bNestedMaskingEnabled), &UAvaMask2DReadModifier::OnNestedMaskingEnabledChanged }
+	{ GET_MEMBER_NAME_CHECKED(UAvaMask2DReadModifier, AdditionalChannels), &UAvaMask2DReadModifier::OnAdditionalChannelsChanged },
 };
 #endif
 
@@ -57,15 +59,6 @@ void UAvaMask2DReadModifier::PostEditChangeProperty(FPropertyChangedEvent& InPro
 	UAvaMask2DReadModifier::PropertyChangeDispatcher.OnPropertyChanged(this, InPropertyChangedEvent);
 }
 
-void UAvaMask2DReadModifier::SetNestedMaskingEnabled(const bool bInNestedMaskingEnabled)
-{
-	if (bNestedMaskingEnabled != bInNestedMaskingEnabled)
-	{
-		bNestedMaskingEnabled = bInNestedMaskingEnabled;
-		OnNestedMaskingEnabledChanged();
-	}
-}
-
 void UAvaMask2DReadModifier::OnMaterialCompiled(UMaterialInterface* InMaterial)
 {
 	if (!InMaterial)
@@ -96,19 +89,39 @@ void UAvaMask2DReadModifier::Apply()
 {
 	Super::Apply();
 
-	if (UAvaMask2DReadModifierShared* Shared = GetShared<UAvaMask2DReadModifierShared>())
+	TOptional<FActorModifierCoreStatus> ModifierStatus;
+	if (AdditionalChannels.Num() > 3)
+	{
+		ModifierStatus = FActorModifierCoreStatus(EActorModifierCoreStatus::Warning, 
+			LOCTEXT("CanvasLimitExceededWarning", "More than 4 channels were gathered.\n"
+				"Only the main channel and the first 3 valid additional channels will be applied."));
+	}
+
+	TArray<TNotNull<const UGeometryMaskCanvas*>, TInlineAllocator<1>> Canvases;
+	if (const UGeometryMaskCanvas* Canvas = GetCurrentCanvas())
 	{
-		TArray<TNotNull<const UGeometryMaskCanvas*>, TInlineAllocator<1>> Canvases;
-		Shared->GatherCanvases(this, Canvases);
+		Canvases.Add(Canvas);
+	}
 
-		if (!Canvases.IsEmpty())
+	const ULevel* Level = nullptr;
+	if (UGeometryMaskWorldSubsystem* MaskSubsystem = GetGeometryMaskWorldSubsystem(&Level))
+	{
+		for (FName AdditionalChannel : AdditionalChannels)
 		{
-			for (TPair<TObjectPtr<UObject>, TInstancedStruct<FAvaMaterialContainerState>>& Pair : MaterialContainerStates)
+			if (UGeometryMaskCanvas* Canvas = MaskSubsystem->GetNamedCanvas(Level, AdditionalChannel))
 			{
-				if (Pair.Key)
-				{
-					ApplyRead(Pair.Key, Canvases);
-				}
+				Canvases.Add(Canvas);
+			}
+		}
+	}
+
+	if (!Canvases.IsEmpty())
+	{
+		for (TPair<TObjectPtr<UObject>, TInstancedStruct<FAvaMaterialContainerState>>& Pair : MaterialContainerStates)
+		{
+			if (Pair.Key)
+			{
+				ApplyRead(Pair.Key, Canvases);
 			}
 		}
 	}
@@ -187,15 +200,16 @@ bool UAvaMask2DReadModifier::ApplyRead(TNotNull<UObject*> InMaterialContainer, T
 		return true;
 	}
 
-	const UGeometryMaskCanvasResource* const FirstCanvasResource = InCanvases[0]->GetResource();
-	if (!ensure(FirstCanvasResource))
+	const TSharedPtr<FGeometryMaskCanvasSharedData> CanvasSharedData = InCanvases[0]->GetSharedData();
+	if (!ensure(CanvasSharedData.IsValid()))
 	{
 		return true;
 	}
 
 	// using only the first canvas parameters as it is consistent across the other canvases.
-	const FIntPoint ViewportSize = FirstCanvasResource->GetMaxViewportSize();
-	const FVector2f ViewportPadding = FVector2f::One() - FVector2f(ViewportSize + FirstCanvasResource->GetViewportPadding()) / ViewportSize;
+	const FVector2f ViewportPadding = CanvasSharedData->TextureSize == FIntPoint::ZeroValue 
+		? FVector2f::Zero()
+		: FVector2f::One() - (FVector2f(CanvasSharedData->TextureSize + CanvasSharedData->ViewportPadding) / CanvasSharedData->TextureSize);
 
 	FAvaMask2DMaskState MaskedState;
 
@@ -209,11 +223,10 @@ bool UAvaMask2DReadModifier::ApplyRead(TNotNull<UObject*> InMaterialContainer, T
 	MaskedState.MaterialParameters.bApplyFeathering = bUseFeathering;
 	MaskedState.MaterialParameters.OuterFeatherRadius = OuterFeatherRadius;
 	MaskedState.MaterialParameters.InnerFeatherRadius = InnerFeatherRadius;
-	MaskedState.MaterialParameters.RenderTarget = FirstCanvasResource->GetRenderTarget();
+	MaskedState.MaterialParameters.RenderTarget = InCanvases[0]->GetRenderTarget();
 
 	// Fill the mask indices of each of the canvases (max 4 because of linear color)
-	static_assert(UAvaMask2DReadModifierShared::MaxCanvasGatherCount == 4);
-	for (int32 Index = 0; Index < UAvaMask2DReadModifierShared::MaxCanvasGatherCount; ++Index)
+	for (int32 Index = 0; Index < 4; ++Index)
 	{
 		float& MaskIndexComponent = MaskedState.MaterialParameters.MaskIndices.Component(Index);
 		if (InCanvases.IsValidIndex(Index) && ensure(InCanvases[Index]->GetRenderTarget() == MaskedState.MaterialParameters.RenderTarget))
@@ -321,11 +334,6 @@ void UAvaMask2DReadModifier::OnModifierEnabled(EActorModifierCoreEnableReason In
 {
 	Super::OnModifierEnabled(InReason);
 
-	if (UAvaMask2DReadModifierShared* Shared = GetShared<UAvaMask2DReadModifierShared>(/*bInCreateIfNone*/true))
-	{
-		Shared->RegisterModifier(this);
-	}
-
 #if WITH_EDITOR
 	UMaterial::OnMaterialCompilationFinished().AddUObject(this, &UAvaMask2DReadModifier::OnMaterialCompiled);
 	FEditorDelegates::EndPIE.AddUObject(this, &UAvaMask2DReadModifier::OnPIEEnded);
@@ -336,11 +344,6 @@ void UAvaMask2DReadModifier::OnModifierDisabled(EActorModifierCoreDisableReason
 {
 	Super::OnModifierDisabled(InReason);
 
-	if (UAvaMask2DReadModifierShared* Shared = GetShared<UAvaMask2DReadModifierShared>(/*bInCreateIfNone*/false))
-	{
-		Shared->UnregisterModifier(this);
-	}
-
 #if WITH_EDITOR
 	UMaterial::OnMaterialCompilationFinished().RemoveAll(this);
 	FEditorDelegates::EndPIE.RemoveAll(this);
@@ -441,7 +444,7 @@ void UAvaMask2DReadModifier::OnBaseOpacityChanged()
 	MarkModifierDirty();
 }
 
-void UAvaMask2DReadModifier::OnNestedMaskingEnabledChanged()
+void UAvaMask2DReadModifier::OnAdditionalChannelsChanged()
 {
 	MarkModifierDirty();
 }
diff --git a/Engine/Plugins/VirtualProduction/Avalanche/Source/AvalancheMask/Public/Mask2D/AvaMask2DBaseModifier.h b/Engine/Plugins/VirtualProduction/Avalanche/Source/AvalancheMask/Public/Mask2D/AvaMask2DBaseModifier.h
--- a/Engine/Plugins/VirtualProduction/Avalanche/Source/AvalancheMask/Public/Mask2D/AvaMask2DBaseModifier.h
+++ b/Engine/Plugins/VirtualProduction/Avalanche/Source/AvalancheMask/Public/Mask2D/AvaMask2DBaseModifier.h
@@ -23,6 +23,7 @@ class IAvaMaskMaterialHandle;
 class UAvaMaskMaterialInstanceSubsystem;
 class UAvaObjectHandleSubsystem;
 class UGeometryMaskCanvas;
+class UGeometryMaskWorldSubsystem;
 class UMaterialInterface;
 class UPrimitiveComponent;
 class UTexture;
@@ -179,6 +180,9 @@ protected:
 	/** Returns true if parent channel was found */
 	UE_API bool TryResolveParentChannel();
 
+	/** Returns world subsystem and, optionally, the level for this modifier */
+	UGeometryMaskWorldSubsystem* GetGeometryMaskWorldSubsystem(const ULevel** OutLevel = nullptr) const;
+
 	UE_API void TryResolveCanvas();
 
 	UE_DEPRECATED(5.8, "ActorData has been deprecated. TryResolveCanvasTexture with actor data is no longer supported.")
diff --git a/Engine/Plugins/VirtualProduction/Avalanche/Source/AvalancheMask/Public/Mask2D/AvaMask2DModifierShared.h b/Engine/Plugins/VirtualProduction/Avalanche/Source/AvalancheMask/Public/Mask2D/AvaMask2DModifierShared.h
deleted file mode 100644
diff --git a/Engine/Plugins/VirtualProduction/Avalanche/Source/AvalancheMask/Public/Mask2D/AvaMask2DReadModifier.h b/Engine/Plugins/VirtualProduction/Avalanche/Source/AvalancheMask/Public/Mask2D/AvaMask2DReadModifier.h
--- a/Engine/Plugins/VirtualProduction/Avalanche/Source/AvalancheMask/Public/Mask2D/AvaMask2DReadModifier.h
+++ b/Engine/Plugins/VirtualProduction/Avalanche/Source/AvalancheMask/Public/Mask2D/AvaMask2DReadModifier.h
@@ -52,12 +52,6 @@ public:
 #endif
 	//~ End UObject
 
-	bool IsNestedMaskingEnabled() const
-	{
-		return bNestedMaskingEnabled;
-	}
-	UE_API void SetNestedMaskingEnabled(const bool bInNestedMaskingEnabled);
-
 	float GetBaseOpacity() const
 	{
 		return BaseOpacity;
@@ -83,7 +77,7 @@ protected:
 	UE_API EBlendMode GetBlendMode() const;
 
 	UE_API void OnBaseOpacityChanged();
-	UE_API void OnNestedMaskingEnabledChanged();
+	UE_API void OnAdditionalChannelsChanged();
 
 	UE_DEPRECATED(5.8, "SetupMaskReadComponent is deprecated. Mask Read Modifier no longer allocates a component")
 	UE_API virtual void SetupMaskComponent(UActorComponent* InComponent) override;
@@ -104,9 +98,9 @@ protected:
 	/** Stores the container and material states for a given actor */
 	void StoreActorState(AActor* InActor);
 
-	/** Whether this modifier should also take into account parent masks */
-	UPROPERTY(EditInstanceOnly, Getter = "IsNestedMaskingEnabled", Setter = "SetNestedMaskingEnabled", Category = "Mask2D", meta = (AllowPrivateAccess = "true"))
-	bool bNestedMaskingEnabled = false;
+	/** Additional channels to read from */
+	UPROPERTY(EditInstanceOnly, Category = "Mask2D", meta = (EditCondition = "!bUseParentChannel", EditConditionHides, AllowPrivateAccess = "true", DisplayAfter="Channel"))
+	TArray<FName> AdditionalChannels;
 
 	/** Base opacity/alpha to use in Read mode */
 	UPROPERTY(EditInstanceOnly, BlueprintReadWrite, Getter, Setter, Category = "Mask2D", meta = (ClampMin = 0.0, ClampMax = 1.0, AllowPrivateAccess = "true"))
diff --git a/Engine/Plugins/VirtualProduction/GeometryMask/Shaders/Private/GeometryMaskBlurCS.usf b/Engine/Plugins/VirtualProduction/GeometryMask/Shaders/Private/GeometryMaskBlurCS.usf
--- a/Engine/Plugins/VirtualProduction/GeometryMask/Shaders/Private/GeometryMaskBlurCS.usf
+++ b/Engine/Plugins/VirtualProduction/GeometryMask/Shaders/Private/GeometryMaskBlurCS.usf
@@ -7,27 +7,24 @@
 #define MAX_SAMPLES 127
 #define HALF_MAX_SAMPLES MAX_SAMPLES / 2
 
-DECLARE_SCALAR_ARRAY(int, SampleCount, 4);
-float4 WeightsAndOffsets[4 * HALF_MAX_SAMPLES];
-float4 BufferSizeAndDirection[4];
+int SampleCount;
+float4 WeightsAndOffsets[HALF_MAX_SAMPLES];
+float4 BufferSizeAndDirection;
 
 uint2 InputDimensions;
 Texture2D InputTexture;
 SamplerState InputSampler;
 RWTexture2D<float4> OutputTexture;
 
-static const float4 ChannelMasks[4] = { float4(1,0,0,1), float4(0,1,0,1), float4(0,0,1,1), float4(0,0,0,1) };
-
-float4 GetSample(int ChannelIdx, float Weight, float Offset, float2 UV)
+float4 GetSample(float Weight, float Offset, float2 UV)
 {
 	const float2 MinUV = float2(0, 0);
 	const float2 MaxUV = float2(1.0, 1.0);
-	const float2 Direction = BufferSizeAndDirection[ChannelIdx].zw;
-	const float2 BufferSize = BufferSizeAndDirection[ChannelIdx].xy;
+	const float2 Direction = BufferSizeAndDirection.zw;
+	const float2 BufferSize = BufferSizeAndDirection.xy;
 	UV *= BufferSize;
-	
+
 	const float2 UVOffset = float2(Offset * BufferSize.x * Direction.x, Offset * BufferSize.y * Direction.y);
-	float4 ChannelMask = ChannelMasks[ChannelIdx];
 
 	return (InputTexture.SampleLevel(InputSampler, clamp(UV + UVOffset, MinUV, MaxUV), 0) * Weight
 		 + InputTexture.SampleLevel(InputSampler, clamp(UV - UVOffset, MinUV, MaxUV), 0) * Weight);
@@ -51,39 +48,31 @@ void MainCS(
 
 	float4 Color = InputTexture.SampleLevel(InputSampler, clamp(UV, MinUV, MaxUV), 0);
 
-    // Perform on each of R,G,B,A
-    [unroll]
-    for (int ChannelIdx = 0; ChannelIdx < 4; ++ChannelIdx)
-    {
-    	int WeightsAndOffsetsIdx = ChannelIdx * HALF_MAX_SAMPLES;
-    	
-    	// First offset is in zw
-	    {
-	    	const float Weight = WeightsAndOffsets[WeightsAndOffsetsIdx].z;
-	    	const float Offset = WeightsAndOffsets[WeightsAndOffsetsIdx].w;
-
-	    	Color += GetSample(ChannelIdx, Weight, Offset, UV);
-	    }
-
-    	int ChannelSampleCount = GET_SCALAR_ARRAY_ELEMENT(SampleCount, ChannelIdx);
-    	for (int SampleIdx = 2; SampleIdx < ChannelSampleCount; SampleIdx += 2)
-    	{
-    		int ParamIdx = SampleIdx / 2;
-    		{
-    			const float Weight = WeightsAndOffsets[WeightsAndOffsetsIdx + ParamIdx].x;
-    			const float Offset = WeightsAndOffsets[WeightsAndOffsetsIdx + ParamIdx].y;
-
-    			Color += GetSample(ChannelIdx, Weight, Offset, UV);
-    		}
-
-    		{
-    			const float Weight = WeightsAndOffsets[WeightsAndOffsetsIdx + ParamIdx].z;
-    			const float Offset = WeightsAndOffsets[WeightsAndOffsetsIdx + ParamIdx].w;
-
-    			Color += GetSample(ChannelIdx, Weight, Offset, UV);
-    		}
-    	}
-    }
+	// First offset is in zw
+	{
+		const float Weight = WeightsAndOffsets[0].z;
+		const float Offset = WeightsAndOffsets[0].w;
+
+		Color += GetSample(Weight, Offset, UV);
+	}
+
+	for (int SampleIdx = 2; SampleIdx < SampleCount; SampleIdx += 2)
+	{
+		int ParamIdx = SampleIdx / 2;
+		{
+			const float Weight = WeightsAndOffsets[ParamIdx].x;
+			const float Offset = WeightsAndOffsets[ParamIdx].y;
+
+			Color += GetSample(Weight, Offset, UV);
+		}
+
+		{
+			const float Weight = WeightsAndOffsets[ParamIdx].z;
+			const float Offset = WeightsAndOffsets[ParamIdx].w;
+
+			Color += GetSample(Weight, Offset, UV);
+		}
+	}
 
 	Color.a = 1.0;
 
diff --git a/Engine/Plugins/VirtualProduction/GeometryMask/Shaders/Private/GeometryMaskJFInitCS.usf b/Engine/Plugins/VirtualProduction/GeometryMask/Shaders/Private/GeometryMaskJFInitCS.usf
--- a/Engine/Plugins/VirtualProduction/GeometryMask/Shaders/Private/GeometryMaskJFInitCS.usf
+++ b/Engine/Plugins/VirtualProduction/GeometryMask/Shaders/Private/GeometryMaskJFInitCS.usf
@@ -5,7 +5,6 @@
 uint2 InputDimensions;
 float2 OneOverInputDimensions;
 float2 UVRatioAdjustment;
-uint4 ChannelMap; // ChannelNum to ChannelIdx (ie. 0 -> Y) 
 Texture2D InputTexture;
 SamplerState InputSampler;
 RWBuffer<float4> OutputBuffer;
@@ -26,10 +25,8 @@ void MainCS(
 	float Y = (float)PixelCoord.y + 0.5;
 	const float2 UV = float2(X, Y) * OneOverInputDimensions * UVRatioAdjustment;
 
-	const uint RowLength = InputDimensions.x * NUM_CHANNELS;	
-	const uint PixelCoord1D = (PixelCoord.y * RowLength) + (PixelCoord.x * NUM_CHANNELS);
-
-	OutputBuffer[PixelCoord1D] = float4(1,1,1,1);
+	const uint RowLength = InputDimensions.x;	
+	const uint PixelCoord1D = (PixelCoord.y * RowLength) + (PixelCoord.x);
 
 	static int SobelX[] =
 	{
@@ -53,52 +50,35 @@ void MainCS(
 	};
 	
 	float4 Input = InputTexture.Load(int3(PixelCoord, 0));
-	float4 Edges = float4(0,0,0,0);
-	int4 Mask = int4(0,0,0,0);
+	float Edges = 0;
+	int Mask = 0;
+
+	OutputBuffer[PixelCoord1D] = float4(0,0,0,0);
+
+	float GradientX = 0;
+	float GradientY = 0;
 
 	[unroll]
-	for (int ChannelIdx = 0; ChannelIdx < NUM_CHANNELS; ++ChannelIdx)
-	{
-		OutputBuffer[PixelCoord1D + ChannelIdx] = float4(0,0,0,0);
-	}
-	
-	// Perform on each of R,G,B,A
-	[unroll]
-	for (int ChannelIdx = 0; ChannelIdx < NUM_CHANNELS; ++ChannelIdx)
+	for (int HorizontalIdx = 0; HorizontalIdx < 3; ++HorizontalIdx)
 	{
-		float GradientX = 0;
-		float GradientY = 0;
-
-		const uint InputChannelIdx = ChannelMap[ChannelIdx];
-    	
 		[unroll]
-		for (int HorizontalIdx = 0; HorizontalIdx < 3; ++HorizontalIdx)
+		for (int VerticalIdx = 0; VerticalIdx < 3; ++VerticalIdx)
 		{
-			[unroll]
-			for (int VerticalIdx = 0; VerticalIdx < 3; ++VerticalIdx)
-			{
-				uint NeighborIdx = (VerticalIdx * KERNEL_SIZE) + HorizontalIdx;
-				uint2 NeighborPixelCoord = PixelCoord + NeighborOffsetIndices[NeighborIdx];
-				NeighborPixelCoord = clamp(NeighborPixelCoord, uint2(0,0), uint2(InputDimensions.x, InputDimensions.y));
+			uint NeighborIdx = (VerticalIdx * KERNEL_SIZE) + HorizontalIdx;
+			uint2 NeighborPixelCoord = PixelCoord + NeighborOffsetIndices[NeighborIdx];
+			NeighborPixelCoord = clamp(NeighborPixelCoord, uint2(0,0), uint2(InputDimensions.x, InputDimensions.y));
 
-				// > 0 ensures treats the input as binary
-				const float NeighborSample = InputTexture.Load(int3(NeighborPixelCoord, 0))[InputChannelIdx] > 0;
+			// > 0 ensures treats the input as binary
+			const float NeighborSample = InputTexture.Load(int3(NeighborPixelCoord, 0))[0] > 0;
 
-				GradientX += SobelX[NeighborIdx] * NeighborSample;
-				GradientY += SobelY[NeighborIdx] * NeighborSample;
-			}
+			GradientX += SobelX[NeighborIdx] * NeighborSample;
+			GradientY += SobelY[NeighborIdx] * NeighborSample;
 		}
-
-		Edges[ChannelIdx] = max(abs(GradientX), abs(GradientY));
-		Mask[ChannelIdx] = min(1, Edges[ChannelIdx]);
 	}
 
+	Edges = max(abs(GradientX), abs(GradientY));
+	Mask = min(1, Edges);
+
 	// @note: This is specific to JFA requirements
-	// For each channel in each pixel, encode (UV, Edge)
-	[unroll]
-	for (int ChannelIdx = 0; ChannelIdx < NUM_CHANNELS; ++ChannelIdx)
-	{
-		const uint InputChannelIdx = ChannelMap[ChannelIdx];
-		OutputBuffer[PixelCoord1D + ChannelIdx] = float4(UV * Mask[ChannelIdx], Edges[ChannelIdx], Input[InputChannelIdx]);
-	}
+	OutputBuffer[PixelCoord1D] = float4(UV * Mask, Edges, Input[0]);
 }
diff --git a/Engine/Plugins/VirtualProduction/GeometryMask/Shaders/Private/GeometryMaskJFStepCS.usf b/Engine/Plugins/VirtualProduction/GeometryMask/Shaders/Private/GeometryMaskJFStepCS.usf
--- a/Engine/Plugins/VirtualProduction/GeometryMask/Shaders/Private/GeometryMaskJFStepCS.usf
+++ b/Engine/Plugins/VirtualProduction/GeometryMask/Shaders/Private/GeometryMaskJFStepCS.usf
@@ -5,8 +5,7 @@
 uint2 InputDimensions;
 float2 OneOverInputDimensions;
 float2 UVRatioAdjustment;
-uint4 StepSize; // per-channel
-uint4 ChannelMask; // whether to calculate for this step
+uint StepSize;
 Buffer<float4> InputBuffer;
 RWBuffer<float4> OutputBuffer;
 
@@ -22,8 +21,8 @@ void MainCS(
 		return;
 	}
 
-	const uint RowLength = InputDimensions.x * NUM_CHANNELS;
-	uint PixelCoord1D = (PixelCoord.y * RowLength) + (PixelCoord.x * NUM_CHANNELS);
+	const uint RowLength = InputDimensions.x;
+	uint PixelCoord1D = (PixelCoord.y * RowLength) + (PixelCoord.x);
 
 	static uint2 NeighborOffsetIndices[] =
 	{
@@ -32,71 +31,45 @@ void MainCS(
 		{-1, -1},  {0, -1},  {1, -1}
 	};
 
-#if NUM_CHANNELS == 1
-	float MinDistance[NUM_CHANNELS] = { 100000000 };
-	float4 MinSample[NUM_CHANNELS] = { float4(0,0,0,0) };
-#elif NUM_CHANNELS == 2
-	float MinDistance[NUM_CHANNELS] = { 100000000, 100000000 };
-	float4 MinSample[NUM_CHANNELS] = { float4(0,0,0,0), float4(0,0,0,0) };
-#elif NUM_CHANNELS == 3
-	float MinDistance[NUM_CHANNELS] = { 100000000, 100000000, 100000000 };
-	float4 MinSample[NUM_CHANNELS] = { float4(0,0,0,0), float4(0,0,0,0), float4(0,0,0,0) };
-#elif NUM_CHANNELS == 4
-	float MinDistance[NUM_CHANNELS] = { 100000000, 100000000, 100000000, 100000000 };
-	float4 MinSample[NUM_CHANNELS] = { float4(0,0,0,0), float4(0,0,0,0), float4(0,0,0,0), float4(0,0,0,0) };
-#endif
+	float MinDistance = 100000000;
+	float4 MinSample = float4(0,0,0,0);
 
 	float X = (float)PixelCoord.x + 0.5;
 	float Y = (float)PixelCoord.y + 0.5;
 	const float2 UV = float2(X, Y) * OneOverInputDimensions * UVRatioAdjustment;
 
-	// Perform on each of R,G,B,A
+	float4 Input = InputBuffer[PixelCoord1D];
+	const uint2 StepSize2 = uint2(StepSize, StepSize);
+
 	[unroll]
-	for (int ChannelIdx = 0; ChannelIdx < NUM_CHANNELS; ++ChannelIdx)
+	for (int HorizontalIdx = 0; HorizontalIdx < 3; ++HorizontalIdx)
 	{
-		float4 Input = InputBuffer[PixelCoord1D + ChannelIdx];
-		const uint2 StepSizeWithChannel = StepSize[ChannelIdx] * uint2(1, 1);
-
-		// Check if we should calculate this channel (for this step)
-		if (ChannelMask[ChannelIdx] == 0)
-		{
-			MinSample[ChannelIdx] = Input[ChannelIdx];
-			continue;
-		}
-
-		MinDistance[ChannelIdx] = 1000000000;
-		MinSample[ChannelIdx] = float4(0,0,0,0);
-
 		[unroll]
-		for (int HorizontalIdx = 0; HorizontalIdx < 3; ++HorizontalIdx)
+		for (int VerticalIdx = 0; VerticalIdx < 3; ++VerticalIdx)
 		{
-			[unroll]
-			for (int VerticalIdx = 0; VerticalIdx < 3; ++VerticalIdx)
-			{
-				uint NeighborIdx = (VerticalIdx * KERNEL_SIZE) + HorizontalIdx;
-				uint2 NeighborPixelCoord = PixelCoord + (NeighborOffsetIndices[NeighborIdx] * StepSizeWithChannel);
+			uint NeighborIdx = (VerticalIdx * KERNEL_SIZE) + HorizontalIdx;
+			uint2 NeighborPixelCoord = PixelCoord + (NeighborOffsetIndices[NeighborIdx] * StepSize2);
 
-				// Prevents border banding
-				NeighborPixelCoord = min(InputDimensions - 1, max(uint2(1,1), NeighborPixelCoord));
+			// Prevents border banding
+			NeighborPixelCoord = min(InputDimensions - 1, max(uint2(1,1), NeighborPixelCoord));
 
-				uint NeighborPixelCoord1D = (NeighborPixelCoord.y * RowLength) + (NeighborPixelCoord.x * NUM_CHANNELS);
+			uint NeighborPixelCoord1D = (NeighborPixelCoord.y * RowLength) + (NeighborPixelCoord.x);
 
-				float4 NeighborSample = InputBuffer[NeighborPixelCoord1D + ChannelIdx];
-				if (NeighborSample.x + NeighborSample.y == 0.0)
-				{
-					continue;
-				}
+			float4 NeighborSample = InputBuffer[NeighborPixelCoord1D];
+			if (NeighborSample.x + NeighborSample.y == 0.0)
+			{
+				continue;
+			}
 
-				const float Distance = length(NeighborSample.xy - UV);
-				if (Distance < MinDistance[ChannelIdx])
-				{
-					MinDistance[ChannelIdx] = Distance;
-					MinSample[ChannelIdx] = NeighborSample;
-				}
+			const float Distance = length(NeighborSample.xy - UV);
+			if (Distance < MinDistance)
+			{
+				MinDistance = Distance;
+				MinSample = NeighborSample;
 			}
 		}
-
-		MinSample[ChannelIdx].a = Input.a;
-		OutputBuffer[PixelCoord1D + ChannelIdx] = MinSample[ChannelIdx];
 	}
+
+	MinSample.a = Input.a;
+	OutputBuffer[PixelCoord1D] = MinSample;
 }
diff --git a/Engine/Plugins/VirtualProduction/GeometryMask/Shaders/Private/GeometryMaskJFtoDFCS.usf b/Engine/Plugins/VirtualProduction/GeometryMask/Shaders/Private/GeometryMaskJFtoDFCS.usf
--- a/Engine/Plugins/VirtualProduction/GeometryMask/Shaders/Private/GeometryMaskJFtoDFCS.usf
+++ b/Engine/Plugins/VirtualProduction/GeometryMask/Shaders/Private/GeometryMaskJFtoDFCS.usf
@@ -7,8 +7,7 @@
 uint2 InputDimensions;
 float2 OneOverInputDimensions;
 float2 UVRatioAdjustment;
-uint4 ChannelMap; // ChannelNum to ChannelIdx (ie. 0 -> Y)
-float4 StepDistanceMultipliers; // To normalize
+float StepDistanceMultiplier; // To normalize
 Texture2D OriginalInputTexture;
 SamplerState OriginalInputSampler;
 Buffer<float4> InputBuffer;
@@ -26,8 +25,8 @@ void MainCS(
 		return;
 	}
 
-	const uint RowLength = InputDimensions.x * NUM_CHANNELS;	
-	const uint PixelCoord1D = (PixelCoord.y * RowLength) + (PixelCoord.x * NUM_CHANNELS);
+	const uint RowLength = InputDimensions.x;	
+	const uint PixelCoord1D = (PixelCoord.y * RowLength) + (PixelCoord.x);
 
 	float X = (float)PixelCoord.x + 0.5;
 	float Y = (float)PixelCoord.y + 0.5;
@@ -40,30 +39,16 @@ void MainCS(
 	float4 Sample = InputBuffer[PixelCoord1D];
 	Outline.x = saturate(distance(UV, Sample.xy));
 #else
-	// Perform on each of R,G,B,A
-	[unroll]
-	for (int ChannelIdx = 0; ChannelIdx < NUM_CHANNELS; ++ChannelIdx)
-	{
-		const uint OutputChannelIdx = ChannelMap[ChannelIdx];
-		
-		float4 Sample = InputBuffer[PixelCoord1D + ChannelIdx];
-
-		float MaskValue = (Sample.x + Sample.y + Sample.a > 0.001);
+	float4 Sample = InputBuffer[PixelCoord1D];
 
-		float OutlineValue = saturate(distance(UV, Sample.xy) * StepDistanceMultipliers[OutputChannelIdx]);
+	float MaskValue = (Sample.x + Sample.y + Sample.a > 0.001);
 
-		// Set range such that inner goes from 0.0-0.5, outer goes from 0.5-1.0
-		OutlineValue = mad(OutlineValue, 0.5, 0.5);
-		OutlineValue = 1.0 - abs(OutlineValue - saturate(Sample.a));
+	float OutlineValue = saturate(distance(UV, Sample.xy) * StepDistanceMultiplier);
 
-		switch (OutputChannelIdx)
-		{
-		case 0: Mask.x = MaskValue; Outline.x = OutlineValue; break;
-		case 1: Mask.y = MaskValue; Outline.y = OutlineValue; break;
-		case 2: Mask.z = MaskValue; Outline.z = OutlineValue; break;
-		case 3: Mask.w = MaskValue; Outline.w = OutlineValue; break;
-		}
-	}
+	// Set range such that inner goes from 0.0-0.5, outer goes from 0.5-1.0
+	OutlineValue = mad(OutlineValue, 0.5, 0.5);
+	OutlineValue = 1.0 - abs(OutlineValue - saturate(Sample.a));
+	Mask.x = MaskValue; Outline.x = OutlineValue;
 #endif
 
 	Outline = lerp(0, Outline, Mask);
diff --git a/Engine/Plugins/VirtualProduction/GeometryMask/Source/GeometryMask/Private/GeometryMaskCanvas.cpp b/Engine/Plugins/VirtualProduction/GeometryMask/Source/GeometryMask/Private/GeometryMaskCanvas.cpp
--- a/Engine/Plugins/VirtualProduction/GeometryMask/Source/GeometryMask/Private/GeometryMaskCanvas.cpp
+++ b/Engine/Plugins/VirtualProduction/GeometryMask/Source/GeometryMask/Private/GeometryMaskCanvas.cpp
@@ -1,12 +1,12 @@
 ﻿// Copyright Epic Games, Inc. All Rights Reserved.
 
 #include "GeometryMaskCanvas.h"
-
 #include "Algo/RemoveIf.h"
 #include "Engine/CanvasRenderTarget2D.h"
 #include "Engine/Engine.h"
 #include "Engine/World.h"
 #include "GeometryMaskCanvasResource.h"
+#include "GeometryMaskCanvasSharedData.h"
 #include "GeometryMaskModule.h"
 #include "IGeometryMaskWriteInterface.h"
 
@@ -265,7 +265,7 @@ void UGeometryMaskCanvas::UpdateRenderParameters()
 {
 	if (CanvasResource)
 	{
-		CanvasResource->UpdateRenderParameters(EGeometryMaskColorChannel::Red, bApplyBlur, BlurStrength, bApplyFeather, OuterFeatherRadius, InnerFeatherRadius);
+		CanvasResource->UpdateRenderParameters(bApplyBlur, BlurStrength, bApplyFeather, OuterFeatherRadius, InnerFeatherRadius);
 	}
 }
 
@@ -280,7 +280,7 @@ void UGeometryMaskCanvas::Initialize(const FInitParams& InInitParams)
 
 	CanvasResource = NewObject<UGeometryMaskCanvasResource>(this);
 	CanvasResource->OnDrawToCanvas().AddUObject(this, &UGeometryMaskCanvas::OnDrawToCanvas);
-	CanvasResource->Initialize(InInitParams.RenderTarget, InInitParams.SliceIndex);
+	CanvasResource->Initialize(InInitParams.RenderTarget, InInitParams.SliceIndex, InInitParams.SharedData);
 	UpdateRenderParameters();
 }
 
@@ -289,6 +289,15 @@ void UGeometryMaskCanvas::Update(const ULevel* InLevel,FSceneView& InView)
 	RemoveInvalidWriters();
 }
 
+TSharedPtr<FGeometryMaskCanvasSharedData> UGeometryMaskCanvas::GetSharedData() const
+{
+	if (CanvasResource)
+	{
+		return CanvasResource->GetSharedData();
+	}
+	return nullptr;
+}
+
 void UGeometryMaskCanvas::AssignResource(UGeometryMaskCanvasResource* InResource, const EGeometryMaskColorChannel InColorChannel)
 {
 }
diff --git a/Engine/Plugins/VirtualProduction/GeometryMask/Source/GeometryMask/Private/GeometryMaskCanvasResource.cpp b/Engine/Plugins/VirtualProduction/GeometryMask/Source/GeometryMask/Private/GeometryMaskCanvasResource.cpp
--- a/Engine/Plugins/VirtualProduction/GeometryMask/Source/GeometryMask/Private/GeometryMaskCanvasResource.cpp
+++ b/Engine/Plugins/VirtualProduction/GeometryMask/Source/GeometryMask/Private/GeometryMaskCanvasResource.cpp
@@ -8,6 +8,7 @@
 #include "Engine/Level.h"
 #include "Engine/TextureRenderTarget2DArray.h"
 #include "Engine/World.h"
+#include "GeometryMaskCanvasSharedData.h"
 #include "GeometryMaskCanvasUtils.h"
 #include "GeometryMaskModule.h"
 #include "GeometryMaskSettings.h"
@@ -41,12 +42,9 @@ UGeometryMaskCanvasResource::UGeometryMaskCanvasResource()
 	: DefaultDrawingContext()
 {
 	FGeometryMaskPostProcessParameters_Blur PostProcessParameters_Blur;
-	PostProcessParameters_Blur.bPerChannelApplyBlur.Init(false, UE::GeometryMask::MaxNumChannels);
-	PostProcessParameters_Blur.PerChannelBlurStrength.Init(16.0, UE::GeometryMask::MaxNumChannels);
 	PostProcess_Blur = MakeShared<FGeometryMaskPostProcess_Blur>(PostProcessParameters_Blur);
 
 	FGeometryMaskPostProcessParameters_DistanceField PostProcessParameters_DistanceField;
-	PostProcessParameters_DistanceField.bPerChannelCalculateDF.Init(false, UE::GeometryMask::MaxNumChannels);
 	PostProcess_DistanceField = MakeShared<FGeometryMaskPostProcess_DistanceField>(PostProcessParameters_DistanceField);
 }
 
@@ -56,10 +54,11 @@ void UGeometryMaskCanvasResource::BeginDestroy()
 	Super::BeginDestroy();
 }
 
-void UGeometryMaskCanvasResource::Initialize(UTextureRenderTarget2DArray* InRenderTarget, int32 InSliceIndex)
+void UGeometryMaskCanvasResource::Initialize(UTextureRenderTarget2DArray* InRenderTarget, int32 InSliceIndex, const TSharedPtr<FGeometryMaskCanvasSharedData>& InSharedData)
 {
 	RenderTarget = InRenderTarget;
 	RenderTargetSliceIndex = InSliceIndex;
+	SharedData = InSharedData;
 }
 
 const EGeometryMaskColorChannel UGeometryMaskCanvasResource::GetNextAvailableColorChannel() const
@@ -101,7 +100,7 @@ int32 UGeometryMaskCanvasResource::Compact()
 
 int32 UGeometryMaskCanvasResource::GetNumChannelsUsed() const
 {
-	return UE::GeometryMask::MaxNumChannels;
+	return 1;
 }
 
 bool UGeometryMaskCanvasResource::IsAnyChannelUsed() const
@@ -114,39 +113,36 @@ TArray<FGeometryMaskCanvasId> UGeometryMaskCanvasResource::GetDependentCanvasIds
 	return {};
 }
 
+const TSharedPtr<FGeometryMaskCanvasSharedData>& UGeometryMaskCanvasResource::GetSharedData() const
+{
+	return SharedData;
+}
+
 bool UGeometryMaskCanvasResource::UpdateViewportSize()
 {
-	if (!RenderTarget)
+	if (!RenderTarget || !SharedData.IsValid())
 	{
 		return false;
 	}
-	
-	const float SizeMultiplier = GetDefault<UGeometryMaskSettings>()->GetDefaultResolutionMultiplier();
-	
+
+	SharedData->ViewportPadding = FMath::Max(SharedData->ViewportPadding, GetRequiredViewportPadding());
+
 	// Update based on max required of all drawing contexts
-	FIntPoint NewViewportSize(ForceInitToZero);
-	
-	// Accounts for res multiplier, padding
-	FIntPoint NewScaledViewportSize(ForceInitToZero);
-	
 	for (FGeometryMaskDrawingContext& DrawingContext : DrawingContextCache)
 	{
-		NewViewportSize = NewViewportSize.ComponentMax(DrawingContext.ViewportSize);
-		
-		NewScaledViewportSize = NewScaledViewportSize.ComponentMax(
-			(DrawingContext.ViewportSize + GetViewportPadding())
-			* SizeMultiplier);
+		SharedData->ViewportSize = SharedData->ViewportSize.ComponentMax(DrawingContext.ViewportSize);
 	}
 
-	if (NewViewportSize.Size() == 0)
+	if (SharedData->ViewportSize.Size() == 0)
 	{
 		return false;
 	}
-	
-	MaxViewportSize = NewViewportSize;
 
-	int32 SizeX = NewScaledViewportSize.X;
-	int32 SizeY = NewScaledViewportSize.Y;
+	MaxViewportSize = SharedData->ViewportSize;
+
+	const FIntPoint PaddedViewportSize = MaxViewportSize + SharedData->ViewportPadding;
+	int32 SizeX = PaddedViewportSize.X;
+	int32 SizeY = PaddedViewportSize.Y;
 
 	if (const float RatioX = static_cast<float>(SizeX) / UE::GeometryMask::MaxTextureSize;
 		RatioX > 1.0)
@@ -164,6 +160,8 @@ bool UGeometryMaskCanvasResource::UpdateViewportSize()
 		SizeX /= RatioY;
 	}
 
+	SharedData->TextureSize = FIntPoint(SizeX, SizeY);
+
 	if (RenderTarget->SizeX == SizeX && RenderTarget->SizeY == SizeY)
 	{
 		return false;
@@ -176,74 +174,44 @@ bool UGeometryMaskCanvasResource::UpdateViewportSize()
 	return true;
 }
 
-void UGeometryMaskCanvasResource::UpdateRenderParameters(
-	const EGeometryMaskColorChannel InColorChannel,
-	const bool bInApplyBlur,
-	const double InBlurStrength,
-	bool bInApplyFeather,
-	int32 InOuterFeatherRadius,
-	int32 InInnerFeatherRadius)
+void UGeometryMaskCanvasResource::UpdateRenderParameters(const EGeometryMaskColorChannel InColorChannel,const bool bInApplyBlur,const double InBlurStrength,bool bInApplyFeather,int32 InOuterFeatherRadius,int32 InInnerFeatherRadius)
 {
-	const int32 ChannelIdx = static_cast<int32>(InColorChannel);
+	UpdateRenderParameters(bInApplyBlur, InBlurStrength, bInApplyFeather, InOuterFeatherRadius, InInnerFeatherRadius);
+}
 
-	auto LogOutOfBounds = [InColorChannel]
+void UGeometryMaskCanvasResource::UpdateRenderParameters(bool bInApplyBlur, double InBlurStrength, bool bInApplyFeather, int32 InOuterFeatherRadius, int32 InInnerFeatherRadius)
+{
+	if (RenderTargetSliceIndex < 0)
 	{
-		const FStringView ChannelStringView = UE::GeometryMask::ChannelToString(InColorChannel);
-		UE_LOG(LogGeometryMask, Error, TEXT("ColorChannel wasn't valid for setting shader parameters (was '%.*s', expected R, G, B or A"), ChannelStringView.Len(), ChannelStringView.GetData());
-	};
+		UE_LOG(LogGeometryMask, Error, TEXT("RenderTargetSliceIndex wasn't valid for setting shader parameters (was '%d', expected a valid index 0 or above"), RenderTargetSliceIndex);
+		return;
+	}
 
 	FGeometryMaskPostProcessParameters_Blur BlurParameters = PostProcess_Blur->GetParameters();
 	{
-		if (BlurParameters.bPerChannelApplyBlur.IsValidIndex(ChannelIdx))
-		{
-			BlurParameters.bPerChannelApplyBlur[ChannelIdx] = bInApplyBlur && InBlurStrength > 0.0;
-		}
-		else
-		{
-			LogOutOfBounds();
-			return;
-		}
-
-		if (BlurParameters.PerChannelBlurStrength.IsValidIndex(ChannelIdx))
-		{
-			BlurParameters.PerChannelBlurStrength[ChannelIdx] = InBlurStrength;	
-		}
-		else
-		{
-			LogOutOfBounds();
-			return;
-		}
+		BlurParameters.SliceIndex = RenderTargetSliceIndex;
+		BlurParameters.bApplyBlur = bInApplyBlur && InBlurStrength > 0.0;
+		BlurParameters.BlurStrength = InBlurStrength;
 
 		PostProcess_Blur->SetParameters(BlurParameters);
 	}
 
 	FGeometryMaskPostProcessParameters_DistanceField DFParameters = PostProcess_DistanceField->GetParameters();
 	{
-		if (DFParameters.bPerChannelCalculateDF.IsValidIndex(ChannelIdx))
-		{
-			DFParameters.bPerChannelCalculateDF[ChannelIdx] = bInApplyFeather && (InOuterFeatherRadius + InInnerFeatherRadius) > 0;
-		}
-		else
-		{
-			LogOutOfBounds();
-			return;
-		}
-
-		if (DFParameters.PerChannelRadius.IsValidIndex(ChannelIdx))
-		{
-			DFParameters.PerChannelRadius[ChannelIdx] = FMath::Max(InOuterFeatherRadius, InInnerFeatherRadius);	
-		}
-		else
-		{
-			LogOutOfBounds();
-			return;
-		}
+		DFParameters.SliceIndex = RenderTargetSliceIndex;
+		DFParameters.bCalculateDF = bInApplyFeather && (InOuterFeatherRadius + InInnerFeatherRadius) > 0;
+		DFParameters.Radius = FMath::Max(InOuterFeatherRadius, InInnerFeatherRadius);
 
 		PostProcess_DistanceField->SetParameters(DFParameters);
 	}
 
-	bApplyBlur = BlurParameters.bPerChannelApplyBlur.CountSetBits(0) > 0; // if any channels have blur
-	bApplyDF = DFParameters.bPerChannelCalculateDF.CountSetBits(0) > 0; // if any channels have DF
+	bApplyBlur = BlurParameters.bApplyBlur;
+	bApplyDF = DFParameters.bCalculateDF;
+
+	if (SharedData.IsValid())
+	{
+		SharedData->Reset();
+	}
 
 	// Effects may have changed viewport padding
     UpdateViewportSize();
@@ -251,7 +219,12 @@ void UGeometryMaskCanvasResource::UpdateRenderParameters(
 
 void UGeometryMaskCanvasResource::ResetRenderParameters(EGeometryMaskColorChannel InColorChannel)
 {
-	UpdateRenderParameters(InColorChannel, false, 0.0, false, 0, 0);
+	ResetRenderParameters();
+}
+
+void UGeometryMaskCanvasResource::ResetRenderParameters()
+{
+	UpdateRenderParameters(/*bApplyBlur*/false, /*BlurStrength*/0.0, /*bApplyFeather*/false, /*OuterFeatherRadius*/0, /*InnerFeatherRadius*/0);
 }
 
 void UGeometryMaskCanvasResource::SetViewportSize(FGeometryMaskDrawingContext& InDrawingContext, const FIntPoint& InViewportSize)
@@ -270,32 +243,31 @@ const FIntPoint& UGeometryMaskCanvasResource::GetMaxViewportSize() const
 
 int32 UGeometryMaskCanvasResource::GetViewportPadding(const FGeometryMaskDrawingContext& InDrawingContext) const
 {
-	return GetViewportPadding();
+	return GetRequiredViewportPadding();
 }
 
-int32 UGeometryMaskCanvasResource::GetViewportPadding() const
+int32 UGeometryMaskCanvasResource::GetSharedViewportPadding() const
+{
+	return SharedData.IsValid() ? SharedData->ViewportPadding : 0;
+}
+
+int32 UGeometryMaskCanvasResource::GetRequiredViewportPadding() const
 {
 	int32 MaxFeatherRadius = 0;
 	{
 		const FGeometryMaskPostProcessParameters_DistanceField& DistanceFieldParameters = PostProcess_DistanceField->GetParameters();
-		for (int32 ChannelIdx = 0; ChannelIdx < UE::GeometryMask::MaxNumChannels; ++ChannelIdx)
+		if (DistanceFieldParameters.bCalculateDF)
 		{
-			if (DistanceFieldParameters.bPerChannelCalculateDF[ChannelIdx])
-			{
-				MaxFeatherRadius = FMath::Max(MaxFeatherRadius, DistanceFieldParameters.PerChannelRadius[ChannelIdx]);
-			}
+			MaxFeatherRadius = FMath::Max(MaxFeatherRadius, DistanceFieldParameters.Radius);
 		}
 	}
 
 	int32 MaxBlurRadius = 0;
 	{
 		const FGeometryMaskPostProcessParameters_Blur& BlurParameters = PostProcess_Blur->GetParameters();
-		for (int32 ChannelIdx = 0; ChannelIdx < UE::GeometryMask::MaxNumChannels; ++ChannelIdx)
+		if (BlurParameters.bApplyBlur)
 		{
-			if (BlurParameters.bPerChannelApplyBlur[ChannelIdx])
-			{
-				MaxBlurRadius = FMath::Max(MaxBlurRadius, UE::GeometryMask::Internal::ComputeEffectiveKernelSize(BlurParameters.PerChannelBlurStrength[ChannelIdx]));
-			}
+			MaxBlurRadius = FMath::Max(MaxBlurRadius, UE::GeometryMask::Internal::ComputeEffectiveKernelSize(BlurParameters.BlurStrength));
 		}
 	}
 
@@ -358,30 +330,27 @@ void UGeometryMaskCanvasResource::Draw(const ULevel* InLevel, FSceneView& InView
 	FGeometryMaskDrawingContext* const DrawingContext = GetDrawingContextForLevel(InLevel, InViewIndex);
     check(DrawingContext); // Should always be valid - world was already checked
 
-	bool bViewportSizeUpdated = false;
-
 	// Begin
 	{
 		const FIntPoint ViewportSize = InView.UnconstrainedViewRect.Size();
 		if (ViewportSize.SizeSquared() > 0 && DrawingContext->ViewportSize != ViewportSize)
 		{
 			DrawingContext->ViewportSize = ViewportSize;
-			bViewportSizeUpdated = UpdateViewportSize();
 		}
 
-		FMatrix ProjectionMatrix = InView.ViewMatrices.GetProjectionMatrix();
-		UE::GeometryMask::Private::OverscanProjectionMatrix(ProjectionMatrix, MaxViewportSize, GetViewportPadding());
-
-		DrawingContext->ViewProjectionMatrix = InView.ViewMatrices.GetViewMatrix() * ProjectionMatrix;
-		DrawingContext->ViewProjectionMatrix.M[2][2] = UE_KINDA_SMALL_NUMBER; // Prevents div by zero later
-
 		// Update viewport size
-		if (bViewportSizeUpdated || UpdateViewportSize())
+		if (UpdateViewportSize())
 		{
 			FlushRenderingCommands();
 			Canvas->SetRenderTarget_GameThread(RenderTarget->GameThread_GetRenderTargetResource());
 		}
 
+		FMatrix ProjectionMatrix = InView.ViewMatrices.GetProjectionMatrix();
+		UE::GeometryMask::Private::OverscanProjectionMatrix(ProjectionMatrix, MaxViewportSize, /*Padding*/0);
+
+		DrawingContext->ViewProjectionMatrix = InView.ViewMatrices.GetViewMatrix() * ProjectionMatrix;
+		DrawingContext->ViewProjectionMatrix.M[2][2] = UE_KINDA_SMALL_NUMBER; // Prevents div by zero later
+
 		InLevel->OwningWorld->FlushDeferredParameterCollectionInstanceUpdates();
 
 		Canvas->SetParentCanvasSize(FIntPoint(RenderTarget->SizeX, RenderTarget->SizeY));
diff --git a/Engine/Plugins/VirtualProduction/GeometryMask/Source/GeometryMask/Private/GeometryMaskWorldSubsystem.cpp b/Engine/Plugins/VirtualProduction/GeometryMask/Source/GeometryMask/Private/GeometryMaskWorldSubsystem.cpp
--- a/Engine/Plugins/VirtualProduction/GeometryMask/Source/GeometryMask/Private/GeometryMaskWorldSubsystem.cpp
+++ b/Engine/Plugins/VirtualProduction/GeometryMask/Source/GeometryMask/Private/GeometryMaskWorldSubsystem.cpp
@@ -6,6 +6,7 @@
 #include "Engine/Level.h"
 #include "Engine/TextureRenderTarget2DArray.h"
 #include "Engine/World.h"
+#include "GeometryMaskCanvasSharedData.h"
 #include "GeometryMaskCanvasUtils.h"
 #include "GeometryMaskSVE.h"
 #include "SceneViewExtension.h"
@@ -57,16 +58,22 @@ FGeometryMaskLevelState& UGeometryMaskWorldSubsystem::FindOrAddLevelState(const
 
 		LevelState->RenderTarget->ClearColor = FLinearColor::Black;
 		LevelState->RenderTarget->bForceLinearGamma = false;
+		LevelState->RenderTarget->bCanCreateUAV = true;
 		LevelState->RenderTarget->bTargetArraySlicesIndependently = true;
 
 		constexpr int32 Width = 1920;
 		constexpr int32 Height = 1080;
-		constexpr int32 Slices = 3;
+		constexpr int32 Slices = 1;
 		LevelState->RenderTarget->Init(Width, Height, Slices, PF_R16F);
 	}
 	return *LevelState;
 }
 
+FGeometryMaskLevelState::FGeometryMaskLevelState()
+	: SharedData(MakeShared<FGeometryMaskCanvasSharedData>())
+{
+}
+
 int32 FGeometryMaskLevelState::AcquireNextSliceIndex()
 {
 	if (RenderTarget)
@@ -108,6 +115,7 @@ UGeometryMaskCanvas* UGeometryMaskWorldSubsystem::GetNamedCanvas(const ULevel* I
 	InitParams.CanvasName = InName;
 	InitParams.RenderTarget = LevelState.RenderTarget;
 	InitParams.SliceIndex = LevelState.AcquireNextSliceIndex();
+	InitParams.SharedData = LevelState.GetSharedData();
 
 	NewCanvas->Initialize(InitParams);
 	OnGeometryMaskCanvasCreatedDelegate.Broadcast(NewCanvas);
diff --git a/Engine/Plugins/VirtualProduction/GeometryMask/Source/GeometryMask/Private/Shaders/GeometryMaskPostProcess_Blur.cpp b/Engine/Plugins/VirtualProduction/GeometryMask/Source/GeometryMask/Private/Shaders/GeometryMaskPostProcess_Blur.cpp
--- a/Engine/Plugins/VirtualProduction/GeometryMask/Source/GeometryMask/Private/Shaders/GeometryMaskPostProcess_Blur.cpp
+++ b/Engine/Plugins/VirtualProduction/GeometryMask/Source/GeometryMask/Private/Shaders/GeometryMaskPostProcess_Blur.cpp
@@ -10,7 +10,8 @@
 
 namespace UE::GeometryMask::Private
 {
-	static constexpr int32 MaxNumSteps = 127;	
+	static constexpr int32 MaxNumSteps = 127;
+	static constexpr int32 WeightsAndOffsetsMaxNum = 64; // Half of Num Steps
 }
 
 class FGeometryMaskBlurCS : public FGlobalShader
@@ -19,9 +20,9 @@ class FGeometryMaskBlurCS : public FGlobalShader
 	SHADER_USE_PARAMETER_STRUCT(FGeometryMaskBlurCS, FGlobalShader);
 
 	BEGIN_SHADER_PARAMETER_STRUCT(FParameters, )
-		SHADER_PARAMETER_SCALAR_ARRAY(int32, SampleCount, [4])
-		SHADER_PARAMETER_ARRAY(FVector4f, WeightsAndOffsets, [4 * UE::GeometryMask::Private::MaxNumSteps])
-		SHADER_PARAMETER_ARRAY(FVector4f, BufferSizeAndDirection, [4])
+		SHADER_PARAMETER(int32, SampleCount)
+		SHADER_PARAMETER_ARRAY(FVector4f, WeightsAndOffsets, [UE::GeometryMask::Private::WeightsAndOffsetsMaxNum])
+		SHADER_PARAMETER(FVector4f, BufferSizeAndDirection)
 		SHADER_PARAMETER(FIntPoint, InputDimensions)
 		SHADER_PARAMETER_RDG_TEXTURE_SRV(Texture2D, InputTexture)
 		SHADER_PARAMETER_SAMPLER(SamplerState, InputSampler)
@@ -139,12 +140,6 @@ namespace UE::GeometryMask
 	}
 }
 
-FGeometryMaskPostProcessParameters_Blur::FGeometryMaskPostProcessParameters_Blur()
-{
-	bPerChannelApplyBlur.Init(false, UE::GeometryMask::MaxNumChannels);
-	PerChannelBlurStrength.Init(16.0,  UE::GeometryMask::MaxNumChannels);
-}
-
 FGeometryMaskPostProcess_Blur::FGeometryMaskPostProcess_Blur(const FGeometryMaskPostProcessParameters_Blur& InParameters)
 	: Super(InParameters)
 {
@@ -152,60 +147,41 @@ FGeometryMaskPostProcess_Blur::FGeometryMaskPostProcess_Blur(const FGeometryMask
 
 void FGeometryMaskPostProcess_Blur::Execute(FRenderTarget* InTexture)
 {
-	if (!InTexture)
-	{
-		return;
-	}
-
 	// If blur NOT applied to any channel, early-out 
-	if (!Parameters.bPerChannelApplyBlur.Contains(true))
+	if (!InTexture || !Parameters.bApplyBlur)
 	{
 		return;
 	}
-	
+
 	ENQUEUE_RENDER_COMMAND(CanvasRenderTargetMakeCurrentCommand)(
-	[Self = AsShared(), InTexture](FRHICommandListImmediate& InRHICmdList)
-	{
-		Self->Execute_RenderThread(InRHICmdList, InTexture);
-	});
+		[Self = AsShared(), InTexture](FRHICommandListImmediate& InRHICmdList)
+		{
+			Self->Execute_RenderThread(InRHICmdList, InTexture);
+		});
 }
 
-void FGeometryMaskPostProcess_Blur::Execute_RenderThread(
-	FRHICommandListImmediate& InRHICmdList
-	, FRenderTarget* InTexture)
+void FGeometryMaskPostProcess_Blur::Execute_RenderThread(FRHICommandListImmediate& InRHICmdList, FRenderTarget* InTexture)
 {
 	ensure(IsInRenderingThread());
 
-	int32 PerChannelSampleCount[UE::GeometryMask::MaxNumChannels];
-	TArray<FVector4f> PerChannelWeightsAndOffsets[UE::GeometryMask::MaxNumChannels];
+	TArray<FVector4f> WeightsAndOffsets;
 
-	for (int32 ChannelIdx = 0; ChannelIdx < UE::GeometryMask::MaxNumChannels; ++ChannelIdx)
-	{
-		double BlurStrength = Parameters.PerChannelBlurStrength[ChannelIdx];
-
-		int32 BlurKernelSize = UE::GeometryMask::Internal::ComputeEffectiveKernelSize(BlurStrength);
-
-		BlurStrength = FMath::Max(0.5, BlurStrength);
-
-		// If we're not applying blur, store as sample with value 0
-		PerChannelSampleCount[ChannelIdx] =
-			Parameters.bPerChannelApplyBlur[ChannelIdx]
-				? UE::GeometryMask::Private::ComputeBlurWeights(BlurKernelSize, BlurStrength, PerChannelWeightsAndOffsets[ChannelIdx])
-				: 0;
-	}
+	const int32 BlurKernelSize = UE::GeometryMask::Internal::ComputeEffectiveKernelSize(Parameters.BlurStrength);
+	const double BlurStrength = FMath::Max(0.5, Parameters.BlurStrength);
+	const int32 SampleCount = UE::GeometryMask::Private::ComputeBlurWeights(BlurKernelSize, BlurStrength, WeightsAndOffsets);
 
 	FRDGBuilder GraphBuilder(InRHICmdList);
 	{
 		FIntPoint InputSize = InTexture->GetSizeXY();
 
 		FRDGTextureRef InputTexture = InTexture->GetRenderTargetTexture(GraphBuilder);
-		FRDGTextureSRVRef InputTexture_SRV = GraphBuilder.CreateSRV(FRDGTextureSRVDesc(InputTexture));
+		FRDGTextureSRVRef InputTexture_SRV = GraphBuilder.CreateSRV(FRDGTextureSRVDesc::CreateForSlice(InputTexture, Parameters.SliceIndex));
 		
 		FRDGTextureDesc IntermediateTextureDesc(
 			FRDGTextureDesc::Create2D(
 				InputTexture->Desc.Extent,
 				InputTexture->Desc.Format,
-				FClearValueBinding::Green,
+				FClearValueBinding::None,
 				TexCreate_ShaderResource | TexCreate_UAV));
 
 		FRDGTextureRef IntermediateTexture = GraphBuilder.CreateTexture(IntermediateTextureDesc, TEXT("GeometryMaskBlur.IntermediateTexture"));
@@ -216,7 +192,7 @@ void FGeometryMaskPostProcess_Blur::Execute_RenderThread(
 			FRDGTextureDesc::Create2D(
 				InputTexture->Desc.Extent,
 				InputTexture->Desc.Format,
-				FClearValueBinding::Green,
+				FClearValueBinding::None,
 				TexCreate_ShaderResource | TexCreate_UAV));
 
 		FRDGTextureRef OutputTexture = GraphBuilder.CreateTexture(OutputTextureDesc, TEXT("GeometryMaskBlur.OutputTexture"));
@@ -231,19 +207,13 @@ void FGeometryMaskPostProcess_Blur::Execute_RenderThread(
 			
 			FGeometryMaskBlurCS::FParameters* PassParameters = GraphBuilder.AllocParameters<FGeometryMaskBlurCS::FParameters>();
 			{
-				for (int32 ChannelIdx = 0; ChannelIdx < UE::GeometryMask::MaxNumChannels; ++ChannelIdx)
+				const int32 WeightsAndOffsetsNum = FMath::Min(WeightsAndOffsets.Num(), PassParameters->WeightsAndOffsets.Num());
+				for (int32 Index = 0; Index < WeightsAndOffsetsNum; ++Index)
 				{
-					int32 WeightsAndOffsetsIdxOffset = ChannelIdx * UE::GeometryMask::Private::MaxNumSteps;
-					for (int32 WeightAndOffsetIdx = 0; WeightAndOffsetIdx < PerChannelWeightsAndOffsets[ChannelIdx].Num(); ++WeightAndOffsetIdx)
-					{
-						PassParameters->WeightsAndOffsets[WeightsAndOffsetsIdxOffset + WeightAndOffsetIdx] = PerChannelWeightsAndOffsets[ChannelIdx][WeightAndOffsetIdx];
-					}
-
-					PassParameters->BufferSizeAndDirection[ChannelIdx] = FVector4f(1.0f / InputSize.X, 1.0f / InputSize.Y, 1.0f, 0.0f);
-
-					GET_SCALAR_ARRAY_ELEMENT(PassParameters->SampleCount, ChannelIdx) = PerChannelSampleCount[ChannelIdx];
+					PassParameters->WeightsAndOffsets[Index] = WeightsAndOffsets[Index];
 				}
-
+				PassParameters->BufferSizeAndDirection = FVector4f(1.0f / InputSize.X, 1.0f / InputSize.Y, 1.0f, 0.0f);
+				PassParameters->SampleCount = SampleCount;
 				PassParameters->InputDimensions = InputSize;
 				PassParameters->InputTexture = InputTexture_SRV;
 				PassParameters->InputSampler = TStaticSamplerState<SF_Bilinear>::GetRHI();
@@ -267,18 +237,13 @@ void FGeometryMaskPostProcess_Blur::Execute_RenderThread(
 			
 			FGeometryMaskBlurCS::FParameters* PassParameters = GraphBuilder.AllocParameters<FGeometryMaskBlurCS::FParameters>();
 			{
-				for (int32 ChannelIdx = 0; ChannelIdx < UE::GeometryMask::MaxNumChannels; ++ChannelIdx)
+				const int32 WeightsAndOffsetsNum = FMath::Min(WeightsAndOffsets.Num(), PassParameters->WeightsAndOffsets.Num());
+				for (int32 Index = 0; Index < WeightsAndOffsetsNum; ++Index)
 				{
-					int32 WeightsAndOffsetsIdxOffset = ChannelIdx * UE::GeometryMask::Private::MaxNumSteps;
-					for (int32 WeightAndOffsetIdx = 0; WeightAndOffsetIdx < PerChannelWeightsAndOffsets[ChannelIdx].Num(); ++WeightAndOffsetIdx)
-					{
-						PassParameters->WeightsAndOffsets[WeightsAndOffsetsIdxOffset + WeightAndOffsetIdx] = PerChannelWeightsAndOffsets[ChannelIdx][WeightAndOffsetIdx];
-					}
-
-					PassParameters->BufferSizeAndDirection[ChannelIdx] = FVector4f(1.0f / InputSize.X, 1.0f / InputSize.Y, 0.0f, 1.0f);
-					GET_SCALAR_ARRAY_ELEMENT(PassParameters->SampleCount, ChannelIdx) = PerChannelSampleCount[ChannelIdx];
+					PassParameters->WeightsAndOffsets[Index] = WeightsAndOffsets[Index];
 				}
-
+				PassParameters->BufferSizeAndDirection = FVector4f(1.0f / InputSize.X, 1.0f / InputSize.Y, 0.0f, 1.0f);
+				PassParameters->SampleCount = SampleCount;
 				PassParameters->InputDimensions = InputSize;
 				PassParameters->InputTexture = IntermediateTexture_SRV;
 				PassParameters->InputSampler = TStaticSamplerState<SF_Bilinear>::GetRHI();
@@ -300,7 +265,12 @@ void FGeometryMaskPostProcess_Blur::Execute_RenderThread(
 		{
 			RDG_EVENT_SCOPE(GraphBuilder, "GeometryMaskBlurPostCopy");
 
-			AddCopyTexturePass(GraphBuilder, OutputTexture, InputTexture, FRHICopyTextureInfo());
+			FRHICopyTextureInfo CopyTextureInfo;
+			CopyTextureInfo.SourceSliceIndex = 0;
+			CopyTextureInfo.DestSliceIndex = Parameters.SliceIndex; // the dest target is the input texture
+			CopyTextureInfo.NumSlices = 1; // only one slice to copy
+
+			AddCopyTexturePass(GraphBuilder, OutputTexture, InputTexture, CopyTextureInfo);
 		}
 	}
 
diff --git a/Engine/Plugins/VirtualProduction/GeometryMask/Source/GeometryMask/Private/Shaders/GeometryMaskPostProcess_Blur.h b/Engine/Plugins/VirtualProduction/GeometryMask/Source/GeometryMask/Private/Shaders/GeometryMaskPostProcess_Blur.h
--- a/Engine/Plugins/VirtualProduction/GeometryMask/Source/GeometryMask/Private/Shaders/GeometryMaskPostProcess_Blur.h
+++ b/Engine/Plugins/VirtualProduction/GeometryMask/Source/GeometryMask/Private/Shaders/GeometryMaskPostProcess_Blur.h
@@ -16,10 +16,9 @@ namespace UE::GeometryMask::Internal
 
 struct FGeometryMaskPostProcessParameters_Blur
 {
-	FGeometryMaskPostProcessParameters_Blur();
-	
-	TBitArray<TInlineAllocator<UE::GeometryMask::MaxNumChannels>> bPerChannelApplyBlur;
-	TArray<double, TInlineAllocator<UE::GeometryMask::MaxNumChannels>> PerChannelBlurStrength;
+	bool bApplyBlur = false;
+	uint16 SliceIndex = 0;
+	double BlurStrength = 16.0;
 };
 
 class FGeometryMaskPostProcess_Blur
diff --git a/Engine/Plugins/VirtualProduction/GeometryMask/Source/GeometryMask/Private/Shaders/GeometryMaskPostProcess_DistanceField.cpp b/Engine/Plugins/VirtualProduction/GeometryMask/Source/GeometryMask/Private/Shaders/GeometryMaskPostProcess_DistanceField.cpp
--- a/Engine/Plugins/VirtualProduction/GeometryMask/Source/GeometryMask/Private/Shaders/GeometryMaskPostProcess_DistanceField.cpp
+++ b/Engine/Plugins/VirtualProduction/GeometryMask/Source/GeometryMask/Private/Shaders/GeometryMaskPostProcess_DistanceField.cpp
@@ -18,35 +18,16 @@ namespace UE::GeometryMask::Private
 	static constexpr int32 NumNeighbors = 8;
 	static constexpr int32 MaxSteps = 13;
 
-	static FIntVector4 CalculateStepCount(const FIntVector4& InRadiusSizes)
+	static int32 CalculateStepCount(double InRadiusSize)
 	{
-		return FIntVector4(
-			InRadiusSizes[0] == 0 ? 0 : FMath::CeilToInt32(FMath::Log2(static_cast<double>(InRadiusSizes[0]))) - 1,
-			InRadiusSizes[1] == 0 ? 0 : FMath::CeilToInt32(FMath::Log2(static_cast<double>(InRadiusSizes[1]))) - 1,
-			InRadiusSizes[2] == 0 ? 0 : FMath::CeilToInt32(FMath::Log2(static_cast<double>(InRadiusSizes[2]))) - 1,
-			InRadiusSizes[3] == 0 ? 0 : FMath::CeilToInt32(FMath::Log2(static_cast<double>(InRadiusSizes[3]))) - 1);
+		return InRadiusSize == 0 ? 0 : FMath::CeilToInt32(FMath::Log2(InRadiusSize)) - 1;
 	}
 
-	static FIntVector4 CalculateStepSize(int32 InStepIdx, const FIntVector4& InMaxSteps, FIntVector4& OutMask)
+	static uint32 CalculateStepSize(int32 InStepIdx, int32 InMaxSteps)
 	{
-		FIntVector4 StepDelta = InMaxSteps - FIntVector4(InStepIdx, InStepIdx, InStepIdx, InStepIdx);
-		StepDelta.X = FMath::Min(InMaxSteps.X, StepDelta.X);
-		StepDelta.Y = FMath::Min(InMaxSteps.Y, StepDelta.Y);
-		StepDelta.Z = FMath::Min(InMaxSteps.Z, StepDelta.Z);
-		StepDelta.W = FMath::Min(InMaxSteps.W, StepDelta.W);
-
-		FIntVector4 StepSize(
-			FMath::CeilToInt32(FMath::Pow<double>(2, StepDelta.X)),
-			FMath::CeilToInt32(FMath::Pow<double>(2, StepDelta.Y)),
-			FMath::CeilToInt32(FMath::Pow<double>(2, StepDelta.Z)),
-			FMath::CeilToInt32(FMath::Pow<double>(2, StepDelta.W)));
-
-		OutMask[0] = StepSize.X > 0;
-		OutMask[1] = StepSize.Y > 0;
-		OutMask[2] = StepSize.Z > 0;
-		OutMask[3] = StepSize.W > 0;
-
-		return StepSize;
+		int32 StepDelta = InMaxSteps - InStepIdx;
+		StepDelta = FMath::Min(InMaxSteps, StepDelta);
+		return 1 << StepDelta;
 	}
 }
 
@@ -56,11 +37,9 @@ class FGeometryMaskJFInitCSBase : public FGlobalShader
 	SHADER_USE_PARAMETER_STRUCT(FGeometryMaskJFInitCSBase, FGlobalShader);
 
 	BEGIN_SHADER_PARAMETER_STRUCT(FParameters, )
-		SHADER_PARAMETER_ARRAY(FVector4f, Offsets, [UE::GeometryMask::MaxNumChannels * (UE::GeometryMask::Private::NumNeighbors + 1)])
 		SHADER_PARAMETER(FIntPoint, InputDimensions)
 		SHADER_PARAMETER(FVector2f, OneOverInputDimensions)
 		SHADER_PARAMETER(FVector2f, UVRatioAdjustment)
-		SHADER_PARAMETER(FIntVector4, ChannelMap)
 		SHADER_PARAMETER_RDG_TEXTURE_SRV(Texture2D, InputTexture)
 		SHADER_PARAMETER_SAMPLER(SamplerState, InputSampler)
 		SHADER_PARAMETER_RDG_BUFFER_UAV(RWBuffer<float4>, OutputBuffer)
@@ -87,20 +66,12 @@ public:
 		FGlobalShader::ModifyCompilationEnvironment(InParameters, OutEnvironment);
 
 		OutEnvironment.SetDefine(TEXT("TILE_SIZE"), FComputeShaderUtils::kGolden2DGroupSize);
-		OutEnvironment.SetDefine(TEXT("NUM_CHANNELS"), NumChannels);
 		OutEnvironment.SetDefine(TEXT("KERNEL_SIZE"), 3);
 	}
 };
 
 using FGeometryMaskJFInitCS1 = TGeometryMaskJFInitCS<1>;
-using FGeometryMaskJFInitCS2 = TGeometryMaskJFInitCS<2>;
-using FGeometryMaskJFInitCS3 = TGeometryMaskJFInitCS<3>;
-using FGeometryMaskJFInitCS4 = TGeometryMaskJFInitCS<4>;
-
 IMPLEMENT_SHADER_TYPE(template<>, FGeometryMaskJFInitCS1, TEXT("/Plugin/GeometryMask/Private/GeometryMaskJFInitCS.usf"), TEXT("MainCS"), SF_Compute);
-IMPLEMENT_SHADER_TYPE(template<>, FGeometryMaskJFInitCS2, TEXT("/Plugin/GeometryMask/Private/GeometryMaskJFInitCS.usf"), TEXT("MainCS"), SF_Compute);
-IMPLEMENT_SHADER_TYPE(template<>, FGeometryMaskJFInitCS3, TEXT("/Plugin/GeometryMask/Private/GeometryMaskJFInitCS.usf"), TEXT("MainCS"), SF_Compute);
-IMPLEMENT_SHADER_TYPE(template<>, FGeometryMaskJFInitCS4, TEXT("/Plugin/GeometryMask/Private/GeometryMaskJFInitCS.usf"), TEXT("MainCS"), SF_Compute);
 
 class FGeometryMaskJFStepCSBase : public FGlobalShader
 {
@@ -108,8 +79,7 @@ class FGeometryMaskJFStepCSBase : public FGlobalShader
 	SHADER_USE_PARAMETER_STRUCT(FGeometryMaskJFStepCSBase, FGlobalShader);
 
 	BEGIN_SHADER_PARAMETER_STRUCT(FParameters, )
-		SHADER_PARAMETER(FIntVector4, StepSize)
-		SHADER_PARAMETER(FIntVector4, ChannelMask)
+		SHADER_PARAMETER(uint32, StepSize)
 		SHADER_PARAMETER(FIntPoint, InputDimensions)
 		SHADER_PARAMETER(FVector2f, OneOverInputDimensions)
 		SHADER_PARAMETER(FVector2f, UVRatioAdjustment)
@@ -138,20 +108,11 @@ public:
 	static void ModifyCompilationEnvironment(const FGlobalShaderPermutationParameters& InParameters, FShaderCompilerEnvironment& OutEnvironment)
 	{
 		FGeometryMaskJFStepCSBase::ModifyCompilationEnvironment(InParameters, OutEnvironment);
-
-		OutEnvironment.SetDefine(TEXT("NUM_CHANNELS"), NumChannels);
 	}
 };
 
 using FGeometryMaskJFStepCS1 = TGeometryMaskJFStepCS<1>;
-using FGeometryMaskJFStepCS2 = TGeometryMaskJFStepCS<2>;
-using FGeometryMaskJFStepCS3 = TGeometryMaskJFStepCS<3>;
-using FGeometryMaskJFStepCS4 = TGeometryMaskJFStepCS<4>;
-
 IMPLEMENT_SHADER_TYPE(template<>, FGeometryMaskJFStepCS1, TEXT("/Plugin/GeometryMask/Private/GeometryMaskJFStepCS.usf"), TEXT("MainCS"), SF_Compute);
-IMPLEMENT_SHADER_TYPE(template<>, FGeometryMaskJFStepCS2, TEXT("/Plugin/GeometryMask/Private/GeometryMaskJFStepCS.usf"), TEXT("MainCS"), SF_Compute);
-IMPLEMENT_SHADER_TYPE(template<>, FGeometryMaskJFStepCS3, TEXT("/Plugin/GeometryMask/Private/GeometryMaskJFStepCS.usf"), TEXT("MainCS"), SF_Compute);
-IMPLEMENT_SHADER_TYPE(template<>, FGeometryMaskJFStepCS4, TEXT("/Plugin/GeometryMask/Private/GeometryMaskJFStepCS.usf"), TEXT("MainCS"), SF_Compute);
 
 class FGeometryMaskJFtoDFCSBase : public FGlobalShader
 {
@@ -162,8 +123,7 @@ class FGeometryMaskJFtoDFCSBase : public FGlobalShader
 		SHADER_PARAMETER(FIntPoint, InputDimensions)
 		SHADER_PARAMETER(FVector2f, OneOverInputDimensions)
 		SHADER_PARAMETER(FVector2f, UVRatioAdjustment)
-		SHADER_PARAMETER(FIntVector4, ChannelMap)
-		SHADER_PARAMETER(FVector4f, StepDistanceMultipliers)
+		SHADER_PARAMETER(float, StepDistanceMultiplier)
 		SHADER_PARAMETER_RDG_TEXTURE_SRV(Texture2D, OriginalInputTexture)
 		SHADER_PARAMETER_SAMPLER(SamplerState, OriginalInputSampler)
 		SHADER_PARAMETER_RDG_BUFFER_SRV(Buffer<float4>, InputBuffer)
@@ -195,20 +155,7 @@ public:
 };
 
 using FGeometryMaskJFtoDFCS1 = TGeometryMaskJFtoDFCS<1>;
-using FGeometryMaskJFtoDFCS2 = TGeometryMaskJFtoDFCS<2>;
-using FGeometryMaskJFtoDFCS3 = TGeometryMaskJFtoDFCS<3>;
-using FGeometryMaskJFtoDFCS4 = TGeometryMaskJFtoDFCS<4>;
-
 IMPLEMENT_SHADER_TYPE(template<>, FGeometryMaskJFtoDFCS1, TEXT("/Plugin/GeometryMask/Private/GeometryMaskJFtoDFCS.usf"), TEXT("MainCS"), SF_Compute);
-IMPLEMENT_SHADER_TYPE(template<>, FGeometryMaskJFtoDFCS2, TEXT("/Plugin/GeometryMask/Private/GeometryMaskJFtoDFCS.usf"), TEXT("MainCS"), SF_Compute);
-IMPLEMENT_SHADER_TYPE(template<>, FGeometryMaskJFtoDFCS3, TEXT("/Plugin/GeometryMask/Private/GeometryMaskJFtoDFCS.usf"), TEXT("MainCS"), SF_Compute);
-IMPLEMENT_SHADER_TYPE(template<>, FGeometryMaskJFtoDFCS4, TEXT("/Plugin/GeometryMask/Private/GeometryMaskJFtoDFCS.usf"), TEXT("MainCS"), SF_Compute);
-
-FGeometryMaskPostProcessParameters_DistanceField::FGeometryMaskPostProcessParameters_DistanceField()
-{
-	bPerChannelCalculateDF.Init(false, UE::GeometryMask::MaxNumChannels);
-	PerChannelRadius.Init(0.0, UE::GeometryMask::MaxNumChannels);
-}
 
 FGeometryMaskPostProcess_DistanceField::FGeometryMaskPostProcess_DistanceField(const FGeometryMaskPostProcessParameters_DistanceField& InParameters)
 	: Super(InParameters)
@@ -222,17 +169,18 @@ void FGeometryMaskPostProcess_DistanceField::Execute(FRenderTarget* InTexture)
 		return;
 	}
 
-	// If blur NOT applied to any channel, early-out 
-	if (!Parameters.bPerChannelCalculateDF.Contains(true))
+	// If blur NOT applied to any channel, early-out
+	if (!Parameters.bCalculateDF)
 	{
+		LastInputSize.Reset();
 		return;
 	}
-	
+
 	ENQUEUE_RENDER_COMMAND(CanvasRenderTargetMakeCurrentCommand)(
-	[Self = AsShared(), InTexture](FRHICommandListImmediate& InRHICmdList)
-	{
-		Self->Execute_RenderThread(InRHICmdList, InTexture);
-	});
+		[Self = AsShared(), InTexture](FRHICommandListImmediate& InRHICmdList)
+		{
+			Self->Execute_RenderThread(InRHICmdList, InTexture);
+		});
 }
 
 void FGeometryMaskPostProcess_DistanceField::Execute_RenderThread(
@@ -252,47 +200,21 @@ void FGeometryMaskPostProcess_DistanceField::Execute_RenderThread(
 		FIntPoint InputSize = InTexture->GetSizeXY();
 		FVector2f OneOverInputSize = FVector2f::One() / InputSize;
 
-		bool bInputSizeChanged = LastInputSize != InputSize;
+		const bool bReallocBuffer = !LastInputSize.IsSet() 
+			|| *LastInputSize != InputSize
+#if WITH_EDITOR
+			|| DebugPass > 0
+#endif
+			|| !StoredInitOutputBuffer.IsValid();
+
 		LastInputSize = InputSize;
 
 		// Scales UV's such that the X axis is always 1.0
 		FVector2f InputHeightOverWidth = FVector2f(1.0f, static_cast<float>(InputSize.Y) / static_cast<float>(InputSize.X));
 
 		FIntVector NumGroups = FComputeShaderUtils::GetGroupCount(InputSize, FIntPoint(8, 8));
-		int32 NumActiveChannels = Parameters.bPerChannelCalculateDF.CountSetBits();
-
-		FIntVector4 ChannelMap(ForceInitToZero);
-		{
-			int32 ActiveChannelIdx = 0;
-			for (int32 ChannelIdx = 0; ChannelIdx < UE::GeometryMask::MaxNumChannels; ++ChannelIdx)
-			{
-				for (int32 ChannelMaskIdx = ActiveChannelIdx; ChannelMaskIdx < Parameters.bPerChannelCalculateDF.Num(); ++ChannelMaskIdx)
-				{
-					if (Parameters.bPerChannelCalculateDF[ChannelMaskIdx])
-					{
-						ActiveChannelIdx = ChannelMaskIdx;
-						break;
-					}
-				}
-
-				ChannelMap[ChannelIdx] = ActiveChannelIdx;
-				++ActiveChannelIdx;
-			}
-		}
-
-		// Used to calculate step size per-channel
-		FIntVector4 SizeFromRadius;
-		for (int32 ChannelIdx = 0; ChannelIdx < UE::GeometryMask::MaxNumChannels; ++ChannelIdx)
-		{
-			SizeFromRadius[ChannelIdx] = Parameters.PerChannelRadius[ChannelIdx];
-		}
-
-		bool bNumActiveChannelsChanged = LastNumActiveChannels != NumActiveChannels;
-		LastNumActiveChannels = NumActiveChannels;
 
-		FIntPoint BufferSize = InputSize;
-		BufferSize.X *= NumActiveChannels;
-		
+		const FIntPoint BufferSize = InputSize;
 		const uint32 BufferLength = BufferSize.X * BufferSize.Y;
 
 		FGlobalShaderMap* GlobalShaderMap = GetGlobalShaderMap(GMaxRHIFeatureLevel);
@@ -300,65 +222,29 @@ void FGeometryMaskPostProcess_DistanceField::Execute_RenderThread(
 		TShaderRef<FGeometryMaskJFInitCSBase> InitComputeShader;
 		TShaderRef<FGeometryMaskJFStepCSBase> StepComputeShader;
 		TShaderRef<FGeometryMaskJFtoDFCSBase> OutputComputeShader;
-		
+
 		{
 			TShaderMapRef<FGeometryMaskJFInitCS1> InitComputeShader1(GlobalShaderMap);
-			TShaderMapRef<FGeometryMaskJFInitCS2> InitComputeShader2(GlobalShaderMap);
-			TShaderMapRef<FGeometryMaskJFInitCS3> InitComputeShader3(GlobalShaderMap);
-			TShaderMapRef<FGeometryMaskJFInitCS4> InitComputeShader4(GlobalShaderMap);
-
 			TShaderMapRef<FGeometryMaskJFStepCS1> StepComputeShader1(GlobalShaderMap);
-			TShaderMapRef<FGeometryMaskJFStepCS2> StepComputeShader2(GlobalShaderMap);
-			TShaderMapRef<FGeometryMaskJFStepCS3> StepComputeShader3(GlobalShaderMap);
-			TShaderMapRef<FGeometryMaskJFStepCS4> StepComputeShader4(GlobalShaderMap);
-
 			TShaderMapRef<FGeometryMaskJFtoDFCS1> OutputComputeShader1(GlobalShaderMap);
-			TShaderMapRef<FGeometryMaskJFtoDFCS2> OutputComputeShader2(GlobalShaderMap);
-			TShaderMapRef<FGeometryMaskJFtoDFCS3> OutputComputeShader3(GlobalShaderMap);
-			TShaderMapRef<FGeometryMaskJFtoDFCS4> OutputComputeShader4(GlobalShaderMap);
 
-			if (NumActiveChannels == 1)
-			{
-				InitComputeShader = InitComputeShader1;
-				StepComputeShader = StepComputeShader1;
-				OutputComputeShader = OutputComputeShader1;
-			}
-			else if (UE::GeometryMask::MaxNumChannels >= 2 && NumActiveChannels == 2)
-			{
-				InitComputeShader = InitComputeShader2;
-				StepComputeShader = StepComputeShader2;
-				OutputComputeShader = OutputComputeShader2;
-			}
-			else if (UE::GeometryMask::MaxNumChannels >= 3 && NumActiveChannels == 3)
-			{
-				InitComputeShader = InitComputeShader3;
-				StepComputeShader = StepComputeShader3;
-				OutputComputeShader = OutputComputeShader3;
-			}
-			else if(UE::GeometryMask::MaxNumChannels == 4 && NumActiveChannels == 4)
-			{
-				InitComputeShader = InitComputeShader4;
-				StepComputeShader = StepComputeShader4;
-				OutputComputeShader = OutputComputeShader4;
-			}
+			InitComputeShader = InitComputeShader1;
+			StepComputeShader = StepComputeShader1;
+			OutputComputeShader = OutputComputeShader1;
 		}
 
 		FRDGTextureRef InputTexture = InTexture->GetRenderTargetTexture(GraphBuilder);
-		FRDGTextureSRVRef InputTexture_SRV = GraphBuilder.CreateSRV(FRDGTextureSRVDesc(InputTexture));
-		FRDGTextureUAVRef InputTexture_UAV = GraphBuilder.CreateUAV(FRDGTextureUAVDesc(InputTexture));
-
-#if WITH_EDITOR
-		bInputSizeChanged = DebugPass > 0;
-#endif
+		FRDGTextureSRVRef InputTexture_SRV = GraphBuilder.CreateSRV(FRDGTextureSRVDesc::CreateForSlice(InputTexture, Parameters.SliceIndex));
+		FRDGTextureUAVRef InputTexture_UAV = GraphBuilder.CreateUAV(FRDGTextureUAVDesc(InputTexture, /*MipLevel*/0, /*Format*/InputTexture->Desc.Format, Parameters.SliceIndex, /*NumSlices*/1));
 
 		FVector4f* BufferData = nullptr;
-		if (bInputSizeChanged || bNumActiveChannelsChanged)
+		if (bReallocBuffer)
 		{
 			BufferData = GraphBuilder.AllocPODArray<FVector4f>(BufferLength);
 		}
 
 		FRDGBufferRef InitOutputBuffer = nullptr;
-		if (bInputSizeChanged || bNumActiveChannelsChanged)
+		if (bReallocBuffer)
 		{
 			FRDGBufferDesc InitOutputBufferDesc =
 				FRDGBufferDesc::CreateBufferDesc(
@@ -376,9 +262,11 @@ void FGeometryMaskPostProcess_DistanceField::Execute_RenderThread(
 			InitOutputBuffer = GraphBuilder.RegisterExternalBuffer(StoredInitOutputBuffer, TEXT("GeometryMaskJFInit.OutputBuffer"));
 		}
 
+		const EPixelFormat BufferFormat = EPixelFormat::PF_FloatRGBA;
+
 		// 1. Init from input binary-ish texture
 		{
-			FRDGBufferUAVRef OutputBuffer_UAV = GraphBuilder.CreateUAV(InitOutputBuffer, EPixelFormat::PF_FloatRGBA);
+			FRDGBufferUAVRef OutputBuffer_UAV = GraphBuilder.CreateUAV(InitOutputBuffer, BufferFormat);
 
 			{
 				RDG_EVENT_SCOPE_STAT(GraphBuilder, GeometryMaskJFInit, "GeometryMaskJFInit");
@@ -391,7 +279,6 @@ void FGeometryMaskPostProcess_DistanceField::Execute_RenderThread(
 					PassParameters->InputDimensions = InputSize;
 					PassParameters->OneOverInputDimensions = OneOverInputSize;
 					PassParameters->UVRatioAdjustment = InputHeightOverWidth;
-					PassParameters->ChannelMap = ChannelMap;
 					PassParameters->InputTexture = InputTexture_SRV;
 					PassParameters->InputSampler = TStaticSamplerState<SF_Bilinear>::GetRHI();
 					PassParameters->OutputBuffer = OutputBuffer_UAV;
@@ -409,19 +296,18 @@ void FGeometryMaskPostProcess_DistanceField::Execute_RenderThread(
 
 		FRDGBufferSRVRef StepOutputBuffer_SRV = nullptr;
 
-		
 #if WITH_EDITOR
 		// DebugPass == 1 - Only sobel
 		if (DebugPass == 1)
 		{
-			StepOutputBuffer_SRV = GraphBuilder.CreateSRV(InitOutputBuffer, EPixelFormat::PF_FloatRGBA);
+			StepOutputBuffer_SRV = GraphBuilder.CreateSRV(InitOutputBuffer, BufferFormat);
 		}
 		else
 #endif
 		// 2. JumpFlood steps
 		{
 			FRDGBufferRef StepIntermediateBufferB = nullptr;
-			if (bInputSizeChanged || bNumActiveChannelsChanged)
+			if (bReallocBuffer)
 			{
 				FRDGBufferDesc StepIntermediateBufferDescB =
 					FRDGBufferDesc::CreateBufferDesc(
@@ -439,18 +325,17 @@ void FGeometryMaskPostProcess_DistanceField::Execute_RenderThread(
 				StepIntermediateBufferB = GraphBuilder.RegisterExternalBuffer(StoredInitOutputBuffer, TEXT("GeometryMaskJFInit.IntermediateBufferB"));
 			}
 
-			FRDGBufferUAVRef StepIntermediateBufferA_UAV = GraphBuilder.CreateUAV(InitOutputBuffer, EPixelFormat::PF_FloatRGBA);
-			FRDGBufferUAVRef StepIntermediateBufferB_UAV = GraphBuilder.CreateUAV(StepIntermediateBufferB, EPixelFormat::PF_FloatRGBA);
+			FRDGBufferUAVRef StepIntermediateBufferA_UAV = GraphBuilder.CreateUAV(InitOutputBuffer, BufferFormat);
+			FRDGBufferUAVRef StepIntermediateBufferB_UAV = GraphBuilder.CreateUAV(StepIntermediateBufferB, BufferFormat);
 
-			FRDGBufferSRVRef StepIntermediateBufferA_SRV = GraphBuilder.CreateSRV(InitOutputBuffer, EPixelFormat::PF_FloatRGBA);
-			FRDGBufferSRVRef StepIntermediateBufferB_SRV = GraphBuilder.CreateSRV(StepIntermediateBufferB, EPixelFormat::PF_FloatRGBA);
+			FRDGBufferSRVRef StepIntermediateBufferA_SRV = GraphBuilder.CreateSRV(InitOutputBuffer, BufferFormat);
+			FRDGBufferSRVRef StepIntermediateBufferB_SRV = GraphBuilder.CreateSRV(StepIntermediateBufferB, BufferFormat);
 
 			StepOutputBuffer_SRV = StepIntermediateBufferA_SRV;
 
 			{
-				FIntVector4 StepCount = UE::GeometryMask::Private::CalculateStepCount(SizeFromRadius);
-				FIntVector4 ChannelMask(ForceInitToZero);
-				int32 MaxStepCount = FMath::Max(FMath::Max(FMath::Max(StepCount.X, StepCount.Y), StepCount.Z), StepCount.W);
+				int32 StepCount = UE::GeometryMask::Private::CalculateStepCount(Parameters.Radius);
+				int32 MaxStepCount = StepCount;
 				int32 BufferIdx = 0;
 
 #if WITH_EDITOR
@@ -478,9 +363,7 @@ void FGeometryMaskPostProcess_DistanceField::Execute_RenderThread(
 						PassParameters->OneOverInputDimensions = OneOverInputSize;
 						PassParameters->UVRatioAdjustment = InputHeightOverWidth;
 
-						PassParameters->StepSize = UE::GeometryMask::Private::CalculateStepSize(StepIdx, StepCount, ChannelMask);
-						PassParameters->ChannelMask = ChannelMask;
-						
+						PassParameters->StepSize = UE::GeometryMask::Private::CalculateStepSize(StepIdx, StepCount);
 						PassParameters->InputBuffer = BufferIdx > 0
 							? StepIntermediateBufferB_SRV
 							: StepIntermediateBufferA_SRV;
@@ -517,18 +400,12 @@ void FGeometryMaskPostProcess_DistanceField::Execute_RenderThread(
 
 				FGeometryMaskJFtoDFCSBase::FParameters* PassParameters = GraphBuilder.AllocParameters<FGeometryMaskJFtoDFCSBase::FParameters>();
 				{
-					FVector4f StepDistanceMultipliers;
-					for (int32 ChannelIdx = 0; ChannelIdx < UE::GeometryMask::MaxNumChannels; ++ChannelIdx)
-					{
-						const float Radius = Parameters.bPerChannelCalculateDF[ChannelIdx] ? FMath::Max(1, Parameters.PerChannelRadius[ChannelIdx]) : InputSize.GetMax();
-						StepDistanceMultipliers[ChannelIdx] = static_cast<float>(InputSize.GetMax()) / Radius; 
-					}
-					
+					const float StepDistanceMultiplier = static_cast<float>(InputSize.GetMax()) / FMath::Max(1, Parameters.Radius);
+
 					PassParameters->InputDimensions = InputSize;
 					PassParameters->OneOverInputDimensions = OneOverInputSize;
 					PassParameters->UVRatioAdjustment = InputHeightOverWidth;
-					PassParameters->ChannelMap = ChannelMap;
-					PassParameters->StepDistanceMultipliers = StepDistanceMultipliers;
+					PassParameters->StepDistanceMultiplier = StepDistanceMultiplier;
 					PassParameters->OriginalInputTexture = InputTexture_SRV;
 					PassParameters->OriginalInputSampler = TStaticSamplerState<SF_Bilinear>::GetRHI();
 					PassParameters->InputBuffer = StepOutputBuffer_SRV;
diff --git a/Engine/Plugins/VirtualProduction/GeometryMask/Source/GeometryMask/Private/Shaders/GeometryMaskPostProcess_DistanceField.h b/Engine/Plugins/VirtualProduction/GeometryMask/Source/GeometryMask/Private/Shaders/GeometryMaskPostProcess_DistanceField.h
--- a/Engine/Plugins/VirtualProduction/GeometryMask/Source/GeometryMask/Private/Shaders/GeometryMaskPostProcess_DistanceField.h
+++ b/Engine/Plugins/VirtualProduction/GeometryMask/Source/GeometryMask/Private/Shaders/GeometryMaskPostProcess_DistanceField.h
@@ -18,10 +18,9 @@ class FRenderTarget;
 
 struct FGeometryMaskPostProcessParameters_DistanceField
 {
-	FGeometryMaskPostProcessParameters_DistanceField();
-	
-	TBitArray<TInlineAllocator<UE::GeometryMask::MaxNumChannels>> bPerChannelCalculateDF;
-	TArray<double, TInlineAllocator<UE::GeometryMask::MaxNumChannels>> PerChannelRadius;
+	bool bCalculateDF = false;
+	uint16 SliceIndex = 0;
+	double Radius = 16.f;
 };
 
 /** */
@@ -41,10 +40,7 @@ protected:
 
 private:
 	/** Cached to selectively resize resources when needed. */
-	FIntPoint LastInputSize{ForceInitToZero};
-
-	int32 LastNumActiveChannels = 0;
-
+	TOptional<FIntPoint> LastInputSize;
 	TRefCountPtr<FRDGPooledBuffer> StoredInitOutputBuffer;
 	TRefCountPtr<FRDGPooledBuffer> StoredStepIntermediateBufferB;
 };
diff --git a/Engine/Plugins/VirtualProduction/GeometryMask/Source/GeometryMask/Public/GeometryMaskCanvas.h b/Engine/Plugins/VirtualProduction/GeometryMask/Source/GeometryMask/Public/GeometryMaskCanvas.h
--- a/Engine/Plugins/VirtualProduction/GeometryMask/Source/GeometryMask/Public/GeometryMaskCanvas.h
+++ b/Engine/Plugins/VirtualProduction/GeometryMask/Source/GeometryMask/Public/GeometryMaskCanvas.h
@@ -14,6 +14,7 @@ class IGeometryMaskWriteInterface;
 class UCanvasRenderTarget2D;
 class UGeometryMaskCanvasResource;
 class UTextureRenderTarget2DArray;
+struct FGeometryMaskCanvasSharedData;
 
 /** Called when Writers becomes non-empty. */
 using FOnGeometryMaskCanvasActivated = TDelegate<void()>;
@@ -119,7 +120,7 @@ public:
 		return OnDeactivatedDelegate;
 	}
 
-	UE_DEPRECATED(5.8, "Use the Initialize() overload with init params")
+	UE_DEPRECATED(5.8, "Initialize with level and canvas name has been deprecated")
 	UE_API void Initialize(const ULevel* InLevel, FName InCanvasName);
 
 	struct FInitParams
@@ -132,12 +133,17 @@ public:
 		UTextureRenderTarget2DArray* RenderTarget = nullptr;
 		/** Index to the slice in the render target to draw to */
 		int32 SliceIndex = -1;
+		/** Shared data between canvases */
+		TSharedPtr<FGeometryMaskCanvasSharedData> SharedData;
 	};
-	UE_API void Initialize(const FInitParams& InInitParams);
+	void Initialize(const FInitParams& InInitParams);
 
 	/** Updates the canvas, intended to be called every frame. */
 	UE_API void Update(const ULevel* InLevel, FSceneView& InView);
 
+	/** Get the data shared among canvases using the same render target */
+	UE_API TSharedPtr<FGeometryMaskCanvasSharedData> GetSharedData() const;
+
 	const UGeometryMaskCanvasResource* GetResource() const
 	{
 		return CanvasResource;
diff --git a/Engine/Plugins/VirtualProduction/GeometryMask/Source/GeometryMask/Public/GeometryMaskCanvasResource.h b/Engine/Plugins/VirtualProduction/GeometryMask/Source/GeometryMask/Public/GeometryMaskCanvasResource.h
--- a/Engine/Plugins/VirtualProduction/GeometryMask/Source/GeometryMask/Public/GeometryMaskCanvasResource.h
+++ b/Engine/Plugins/VirtualProduction/GeometryMask/Source/GeometryMask/Public/GeometryMaskCanvasResource.h
@@ -18,6 +18,7 @@ class UCanvasRenderTarget2D;
 class UTextureRenderTarget2DArray;
 enum class EGeometryMaskColorChannel : uint8;
 enum class ETextureRenderTargetSampleCount : uint8;
+struct FGeometryMaskCanvasSharedData;
 
 using FOnGeometryMaskCanvasDraw = TMulticastDelegate<void(const FGeometryMaskDrawingContext&, FCanvas*)>;
 
@@ -35,7 +36,7 @@ public:
 
 	UE_API virtual void BeginDestroy() override;
 
-	void Initialize(UTextureRenderTarget2DArray* InRenderTarget, int32 InSliceIndex);
+	void Initialize(UTextureRenderTarget2DArray* InRenderTarget, int32 InSliceIndex, const TSharedPtr<FGeometryMaskCanvasSharedData>& InSharedData);
 
 	/** Will return the first available color channel without a canvas assigned. EGeometryMaskColorChannel::None if not available. */
 	UE_DEPRECATED(5.8, "Color channel management has been deprecated as it's no longer needed")
@@ -63,6 +64,9 @@ public:
 	UE_DEPRECATED(5.8, "Color channel management has been deprecated as it's no longer needed")
 	UE_API TArray<FGeometryMaskCanvasId> GetDependentCanvasIds() const;
 
+	/** Gets the data shared between canvases using the same render target*/
+	UE_API const TSharedPtr<FGeometryMaskCanvasSharedData>& GetSharedData() const;
+
 	/** 
 	 * Validates that the render target size is up-to-date with the viewport size and updating it if not
 	 * @return true if the render target size was updated
@@ -75,11 +79,14 @@ public:
 	/** Returns the maximum size that would cover all requested viewports */
 	UE_API const FIntPoint& GetMaxViewportSize() const;
 
-	UE_DEPRECATED(5.8, "GetViewportPadding requiring a drawing context has been deprecated. Use the overload without the drawing context")
+	UE_DEPRECATED(5.8, "GetViewportPadding requiring a drawing context has been deprecated. Use GetRequiredViewportPadding() instead")
 	UE_API int32 GetViewportPadding(const FGeometryMaskDrawingContext& InDrawingContext) const;
 
-	/** Returns required viewport padding, in pixels - determined by certain effects. */
-	UE_API int32 GetViewportPadding() const;
+	/** Returns the current viewport padding */
+	UE_API int32 GetSharedViewportPadding() const;
+
+	/** Returns the minimum required viewport padding, in pixels - determined by certain effects. */
+	UE_API int32 GetRequiredViewportPadding() const;
 
 	UE_DEPRECATED(5.8, "GetRenderTargetTexture (UCanvasRenderTarget2D) has been deprecated in favor of GetRenderTarget (UTextureRenderTarget2DArray)")
 	UE_API UCanvasRenderTarget2D* GetRenderTargetTexture();
@@ -95,11 +102,18 @@ public:
 		return OnDrawToCanvasDelegate;
 	}
 
+	UE_DEPRECATED(5.8, "Color Channel is no longer used. Use the overload that does not take the color channel")
 	UE_API void UpdateRenderParameters(EGeometryMaskColorChannel InColorChannel, bool bInApplyBlur, double InBlurStrength, bool bInApplyFeather, int32 InOuterFeatherRadius, int32 InInnerFeatherRadius);
 
+	UE_API void UpdateRenderParameters(bool bInApplyBlur, double InBlurStrength, bool bInApplyFeather, int32 InOuterFeatherRadius, int32 InInnerFeatherRadius);
+
 	/** Resets the render parameters for the given channel. */
+	UE_DEPRECATED(5.8, "Color Channel is no longer used. Use the overload that does not take the color channel")
 	UE_API void ResetRenderParameters(EGeometryMaskColorChannel InColorChannel);
 
+	/** Resets the render parameters for the resource */
+	UE_API void ResetRenderParameters();
+
 	/** Updates the canvas, intended to be called every frame. */
 	UE_API void Update(const ULevel* InLevel, FSceneView& InView, int32 InViewIndex = 0);
 
@@ -140,7 +154,7 @@ private:
 	UPROPERTY()
 	TObjectPtr<UCanvasRenderTarget2D> RenderTargetTexture_DEPRECATED;
 
-	/** The underlying Render Target. */
+	/** The Render Target that this resource uses */
 	UPROPERTY(DuplicateTransient)
 	TObjectPtr<UTextureRenderTarget2DArray> RenderTarget;
 
@@ -148,6 +162,9 @@ private:
 	UPROPERTY(DuplicateTransient)
 	int16 RenderTargetSliceIndex = -1;
 
+	/** Shared data between canvases using the same render target */
+	TSharedPtr<FGeometryMaskCanvasSharedData> SharedData;
+
 	FOnGeometryMaskCanvasDraw OnDrawToCanvasDelegate;
 
 	TSet<FGeometryMaskDrawingContext> DrawingContextCache;
diff --git a/Engine/Plugins/VirtualProduction/GeometryMask/Source/GeometryMask/Public/GeometryMaskCanvasSharedData.h b/Engine/Plugins/VirtualProduction/GeometryMask/Source/GeometryMask/Public/GeometryMaskCanvasSharedData.h
new file mode 100644
--- /dev/null
+++ b/Engine/Plugins/VirtualProduction/GeometryMask/Source/GeometryMask/Public/GeometryMaskCanvasSharedData.h
@@ -0,0 +1,25 @@
+// Copyright Epic Games, Inc. All Rights Reserved.
+
+#pragma once
+
+#include "HAL/Platform.h"
+#include "Math/IntPoint.h"
+
+struct FGeometryMaskCanvasSharedData
+{
+	void Reset()
+	{
+		ViewportSize = FIntPoint(0, 0);
+		TextureSize = FIntPoint(0, 0);
+		ViewportPadding = 0;
+	}
+
+	/** Viewport size of the canvas */
+	FIntPoint ViewportSize = FIntPoint(0, 0);
+
+	/** Actual size of the canvas render target */
+	FIntPoint TextureSize = FIntPoint(0, 0);
+
+	/** Viewport padding affecting the texture size */
+	int32 ViewportPadding = 0;
+};
diff --git a/Engine/Plugins/VirtualProduction/GeometryMask/Source/GeometryMask/Public/GeometryMaskWorldSubsystem.h b/Engine/Plugins/VirtualProduction/GeometryMask/Source/GeometryMask/Public/GeometryMaskWorldSubsystem.h
--- a/Engine/Plugins/VirtualProduction/GeometryMask/Source/GeometryMask/Public/GeometryMaskWorldSubsystem.h
+++ b/Engine/Plugins/VirtualProduction/GeometryMask/Source/GeometryMask/Public/GeometryMaskWorldSubsystem.h
@@ -12,6 +12,7 @@
 
 class FGeometryMaskSceneViewExtension;
 class ULevel;
+struct FGeometryMaskCanvasSharedData;
 
 using FOnGeometryMaskCanvasCreated = TMulticastDelegate<void(const UGeometryMaskCanvas*)>;
 using FOnGeometryMaskCanvasDestroyed = TMulticastDelegate<void(const FGeometryMaskCanvasId&)>;
@@ -21,6 +22,8 @@ struct FGeometryMaskLevelState
 {
 	GENERATED_BODY()
 
+	UE_API FGeometryMaskLevelState();
+
 	/** All the canvases in the level */
 	UPROPERTY()
 	TMap<FName, TObjectPtr<UGeometryMaskCanvas>> NamedCanvases;
@@ -29,6 +32,11 @@ struct FGeometryMaskLevelState
 	UPROPERTY()
 	TObjectPtr<UTextureRenderTarget2DArray> RenderTarget;
 
+	const TSharedRef<FGeometryMaskCanvasSharedData>& GetSharedData() const
+	{
+		return SharedData;
+	}
+
 	/** Returns the index to the next available slice in the render target */
 	int32 AcquireNextSliceIndex();
 
@@ -36,6 +44,9 @@ private:
 	/** tracker for what the next available slice index is */
 	UPROPERTY()
 	int32 NextSliceIndex = 0;
+
+	/** Shared data for canvases that use the render target */
+	TSharedRef<FGeometryMaskCanvasSharedData> SharedData;
 };
 
 /** Updates the canvases. */
diff --git a/Engine/Plugins/VirtualProduction/GeometryMask/Source/GeometryMaskEditor/Private/Widgets/SGeometryMaskCanvasPreview.cpp b/Engine/Plugins/VirtualProduction/GeometryMask/Source/GeometryMaskEditor/Private/Widgets/SGeometryMaskCanvasPreview.cpp
--- a/Engine/Plugins/VirtualProduction/GeometryMask/Source/GeometryMaskEditor/Private/Widgets/SGeometryMaskCanvasPreview.cpp
+++ b/Engine/Plugins/VirtualProduction/GeometryMask/Source/GeometryMaskEditor/Private/Widgets/SGeometryMaskCanvasPreview.cpp
@@ -269,7 +269,7 @@ void SGeometryMaskCanvasPreview::UpdateBrush(const UGeometryMaskCanvas* InCanvas
 		FVector4f Padding(ForceInitToZero);
 		if (const UGeometryMaskCanvasResource* CanvasResource = InCanvas->GetResource())
 		{
-			FIntPoint ViewportPadding(CanvasResource->GetViewportPadding());
+			FIntPoint ViewportPadding(CanvasResource->GetSharedViewportPadding());
 			FVector2f ViewportPaddingF = FVector2f(ViewportPadding) / CanvasResource->GetMaxViewportSize();
 			Padding.X = ViewportPaddingF.X;
 			Padding.Y = ViewportPaddingF.Y;
