Motion Design: added mask client to keep track which channels are in active use by writers/readers. This way new channels can re-use existing slice indices that are no longer used.

#Synced-CL 50160130

diff --git a/Engine/Plugins/VirtualProduction/Avalanche/Source/AvalancheMask/Private/Mask2D/AvaMask2DBaseModifier.cpp b/Engine/Plugins/VirtualProduction/Avalanche/Source/AvalancheMask/Private/Mask2D/AvaMask2DBaseModifier.cpp
--- a/Engine/Plugins/VirtualProduction/Avalanche/Source/AvalancheMask/Private/Mask2D/AvaMask2DBaseModifier.cpp
+++ b/Engine/Plugins/VirtualProduction/Avalanche/Source/AvalancheMask/Private/Mask2D/AvaMask2DBaseModifier.cpp
@@ -436,6 +436,11 @@ UAvaMask2DBaseModifier* UAvaMask2DBaseModifier::FindMaskModifierOnActor(const AA
 	return nullptr;
 }
 
+bool UAvaMask2DBaseModifier::ForEachUsedCanvasName(TFunctionRef<bool(FName)> InFunc) const
+{
+	return InFunc(GetChannel());
+}
+
 void UAvaMask2DBaseModifier::OnModifierCDOSetup(FActorModifierCoreMetadata& InMetadata)
 {
 	Super::OnModifierCDOSetup(InMetadata);
@@ -623,7 +628,7 @@ void UAvaMask2DBaseModifier::TryResolveCanvas()
 	const ULevel* Level = nullptr;
 	if (UGeometryMaskWorldSubsystem* MaskSubsystem = GetGeometryMaskWorldSubsystem(&Level))
 	{
-		if (UGeometryMaskCanvas* Canvas = MaskSubsystem->GetNamedCanvas(Level, GetChannel()))
+		if (UGeometryMaskCanvas* Canvas = MaskSubsystem->GetNamedCanvas(Level, GetChannel(), this))
 		{
 			LastResolvedCanvasName = Canvas->GetCanvasName();
 			CanvasWeak = Canvas;
diff --git a/Engine/Plugins/VirtualProduction/Avalanche/Source/AvalancheMask/Private/Mask2D/AvaMask2DReadModifier.cpp b/Engine/Plugins/VirtualProduction/Avalanche/Source/AvalancheMask/Private/Mask2D/AvaMask2DReadModifier.cpp
--- a/Engine/Plugins/VirtualProduction/Avalanche/Source/AvalancheMask/Private/Mask2D/AvaMask2DReadModifier.cpp
+++ b/Engine/Plugins/VirtualProduction/Avalanche/Source/AvalancheMask/Private/Mask2D/AvaMask2DReadModifier.cpp
@@ -89,34 +89,35 @@ void UAvaMask2DReadModifier::Apply()
 {
 	Super::Apply();
 
-	TOptional<FActorModifierCoreStatus> ModifierStatus;
-	if (AdditionalChannels.Num() > 3)
-	{
-		ModifierStatus = FActorModifierCoreStatus(EActorModifierCoreStatus::Warning, 
-			LOCTEXT("CanvasLimitExceededWarning", "More than 4 channels were gathered.\n"
-				"Only the main channel and the first 3 valid additional channels will be applied."));
-	}
-
 	TArray<TNotNull<const UGeometryMaskCanvas*>, TInlineAllocator<1>> Canvases;
 	if (const UGeometryMaskCanvas* Canvas = GetCurrentCanvas())
 	{
 		Canvases.Add(Canvas);
-	}
 
-	const ULevel* Level = nullptr;
-	if (UGeometryMaskWorldSubsystem* MaskSubsystem = GetGeometryMaskWorldSubsystem(&Level))
-	{
-		for (FName AdditionalChannel : AdditionalChannels)
+		const ULevel* Level = nullptr;
+		if (UGeometryMaskWorldSubsystem* MaskSubsystem = GetGeometryMaskWorldSubsystem(&Level))
 		{
-			if (UGeometryMaskCanvas* Canvas = MaskSubsystem->GetNamedCanvas(Level, AdditionalChannel))
+			for (FName AdditionalChannel : AdditionalChannels)
 			{
-				Canvases.Add(Canvas);
+				if (UGeometryMaskCanvas* AdditionalCanvas = MaskSubsystem->GetNamedCanvas(Level, AdditionalChannel, this))
+				{
+					Canvases.Add(AdditionalCanvas);
+				}
 			}
 		}
 	}
 
 	if (!Canvases.IsEmpty())
 	{
+		if (Canvases.Num() > MaxCanvasCount && UE_LOG_ACTIVE(LogAvaMask, Warning))
+		{
+			const AActor* Actor = GetModifiedActor();
+			UE_LOG(LogAvaMask, Warning, TEXT("%d canvases were found in Masked Layer (Output) for Actor '%s'. Only the first %d (limit) will be considered.")
+				, Canvases.Num()
+				, Actor ? *Actor->GetActorNameOrLabel() : TEXT("(None)")
+				, MaxCanvasCount);
+		}
+
 		for (TPair<TObjectPtr<UObject>, TInstancedStruct<FAvaMaterialContainerState>>& Pair : MaterialContainerStates)
 		{
 			if (Pair.Key)
@@ -225,8 +226,8 @@ bool UAvaMask2DReadModifier::ApplyRead(TNotNull<UObject*> InMaterialContainer, T
 	MaskedState.MaterialParameters.InnerFeatherRadius = InnerFeatherRadius;
 	MaskedState.MaterialParameters.RenderTarget = InCanvases[0]->GetRenderTarget();
 
-	// Fill the mask indices of each of the canvases (max 4 because of linear color)
-	for (int32 Index = 0; Index < 4; ++Index)
+	// Fill the mask indices of each of the canvases
+	for (int32 Index = 0; Index < MaxCanvasCount; ++Index)
 	{
 		float& MaskIndexComponent = MaskedState.MaterialParameters.MaskIndices.Component(Index);
 		if (InCanvases.IsValidIndex(Index) && ensure(InCanvases[Index]->GetRenderTarget() == MaskedState.MaterialParameters.RenderTarget))
@@ -316,6 +317,24 @@ bool UAvaMask2DReadModifier::ApplyRead(AActor* InActor, FAvaMask2DActorData& InA
 	return true;
 }
 
+bool UAvaMask2DReadModifier::ForEachUsedCanvasName(TFunctionRef<bool(FName)> InFunc) const
+{
+	if (!Super::ForEachUsedCanvasName(InFunc))
+	{
+		return false; 
+	}
+
+	for (FName AdditionalChannel : AdditionalChannels)
+	{
+		if (!InFunc(AdditionalChannel))
+		{
+			return false;
+		}
+	}
+
+	return true;
+}
+
 void UAvaMask2DReadModifier::OnModifierCDOSetup(FActorModifierCoreMetadata& InMetadata)
 {
 	Super::OnModifierCDOSetup(InMetadata);
diff --git a/Engine/Plugins/VirtualProduction/Avalanche/Source/AvalancheMask/Public/Mask2D/AvaMask2DBaseModifier.h b/Engine/Plugins/VirtualProduction/Avalanche/Source/AvalancheMask/Public/Mask2D/AvaMask2DBaseModifier.h
--- a/Engine/Plugins/VirtualProduction/Avalanche/Source/AvalancheMask/Public/Mask2D/AvaMask2DBaseModifier.h
+++ b/Engine/Plugins/VirtualProduction/Avalanche/Source/AvalancheMask/Public/Mask2D/AvaMask2DBaseModifier.h
@@ -5,6 +5,7 @@
 #include "AvaMaskTypes.h"
 #include "AvaPropertyChangeDispatcher.h"
 #include "GeometryMaskTypes.h"
+#include "IGeometryMaskClient.h"
 #include "Modifiers/ActorModifierArrangeBaseModifier.h"
 #include "StructUtils/InstancedStruct.h"
 #include "UObject/ObjectPtr.h"
@@ -56,7 +57,7 @@ struct FAvaMask2DActorData
 
 /** Uses scene actors to create a mask texture and applies it to attached actors */
 UCLASS(Abstract, MinimalAPI)
-class UAvaMask2DBaseModifier : public UActorModifierArrangeBaseModifier
+class UAvaMask2DBaseModifier : public UActorModifierArrangeBaseModifier, public IGeometryMaskClient
 {
 	GENERATED_BODY()
 
@@ -121,6 +122,10 @@ public:
 	UE_API static UAvaMask2DBaseModifier* FindMaskModifierOnActor(const AActor* InActor);
 
 protected:
+	//~ Begin IGeometryMaskClient
+	UE_API virtual bool ForEachUsedCanvasName(TFunctionRef<bool(FName)> InFunc) const override;
+	//~ End IGeometryMaskClient
+
 	//~ Begin UActorModifierCoreBase
 	UE_API virtual void OnModifierCDOSetup(FActorModifierCoreMetadata& InMetadata) override;
 	UE_API virtual void OnModifierAdded(EActorModifierCoreEnableReason InReason) override;
diff --git a/Engine/Plugins/VirtualProduction/Avalanche/Source/AvalancheMask/Public/Mask2D/AvaMask2DReadModifier.h b/Engine/Plugins/VirtualProduction/Avalanche/Source/AvalancheMask/Public/Mask2D/AvaMask2DReadModifier.h
--- a/Engine/Plugins/VirtualProduction/Avalanche/Source/AvalancheMask/Public/Mask2D/AvaMask2DReadModifier.h
+++ b/Engine/Plugins/VirtualProduction/Avalanche/Source/AvalancheMask/Public/Mask2D/AvaMask2DReadModifier.h
@@ -44,7 +44,10 @@ class UAvaMask2DReadModifier : public UAvaMask2DBaseModifier
 	GENERATED_BODY()
 
 	friend class UAvaMask2DReadModifierShared;
-	
+
+	/** Limit canvas count to 4 because a single linear color parameter is used to transmit the indices over. */
+	static constexpr int32 MaxCanvasCount = 4;
+
 public:
 	//~ Begin UObject
 #if WITH_EDITOR
@@ -59,6 +62,10 @@ public:
 	UE_API void SetBaseOpacity(float InBaseOpacity);
 
 protected:
+	//~ Begin IGeometryMaskClient
+	UE_API virtual bool ForEachUsedCanvasName(TFunctionRef<bool(FName)> InFunc) const override;
+	//~ End IGeometryMaskClient
+
 	//~ Begin UActorModifierCoreBase
 	UE_API virtual void OnModifierCDOSetup(FActorModifierCoreMetadata& InMetadata) override;
 	UE_API virtual void OnModifierEnabled(EActorModifierCoreEnableReason InReason) override;
diff --git a/Engine/Plugins/VirtualProduction/GeometryMask/Source/GeometryMask/Private/GeometryMaskCanvasActor.cpp b/Engine/Plugins/VirtualProduction/GeometryMask/Source/GeometryMask/Private/GeometryMaskCanvasActor.cpp
--- a/Engine/Plugins/VirtualProduction/GeometryMask/Source/GeometryMask/Private/GeometryMaskCanvasActor.cpp
+++ b/Engine/Plugins/VirtualProduction/GeometryMask/Source/GeometryMask/Private/GeometryMaskCanvasActor.cpp
@@ -69,7 +69,7 @@ bool AGeometryMaskCanvasActor::TryResolveCanvas()
 	{
 		if (UGeometryMaskWorldSubsystem* Subsystem = GetWorld()->GetSubsystem<UGeometryMaskWorldSubsystem>())
 		{
-			Canvas = Subsystem->GetNamedCanvas(GetLevel(), CanvasName);
+			Canvas = Subsystem->GetNamedCanvas(GetLevel(), CanvasName, this);
 		}
 	}
 	
@@ -105,3 +105,8 @@ void AGeometryMaskCanvasActor::FindWriters()
 		Canvas->AddWriters(Writers);
 	}
 }
+
+bool AGeometryMaskCanvasActor::ForEachUsedCanvasName(TFunctionRef<bool(FName)> InFunc) const
+{
+	return InFunc(CanvasName);
+}
diff --git a/Engine/Plugins/VirtualProduction/GeometryMask/Source/GeometryMask/Private/GeometryMaskReadComponent.cpp b/Engine/Plugins/VirtualProduction/GeometryMask/Source/GeometryMask/Private/GeometryMaskReadComponent.cpp
--- a/Engine/Plugins/VirtualProduction/GeometryMask/Source/GeometryMask/Private/GeometryMaskReadComponent.cpp
+++ b/Engine/Plugins/VirtualProduction/GeometryMask/Source/GeometryMask/Private/GeometryMaskReadComponent.cpp
@@ -11,6 +11,11 @@ void UGeometryMaskReadComponent::SetParameters(FGeometryMaskReadParameters& InPa
 	TryResolveCanvas();
 }
 
+bool UGeometryMaskReadComponent::ForEachUsedCanvasName(TFunctionRef<bool(FName)> InFunc) const
+{
+	return InFunc(Parameters.CanvasName);
+}
+
 bool UGeometryMaskReadComponent::TryResolveCanvas()
 {
 	return TryResolveNamedCanvas(Parameters.CanvasName);
diff --git a/Engine/Plugins/VirtualProduction/GeometryMask/Source/GeometryMask/Private/GeometryMaskTypes.cpp b/Engine/Plugins/VirtualProduction/GeometryMask/Source/GeometryMask/Private/GeometryMaskTypes.cpp
--- a/Engine/Plugins/VirtualProduction/GeometryMask/Source/GeometryMask/Private/GeometryMaskTypes.cpp
+++ b/Engine/Plugins/VirtualProduction/GeometryMask/Source/GeometryMask/Private/GeometryMaskTypes.cpp
@@ -200,7 +200,7 @@ bool UGeometryMaskCanvasReferenceComponentBase::TryResolveNamedCanvas(FName InCa
 
 	if (UGeometryMaskWorldSubsystem* Subsystem = GetWorld()->GetSubsystem<UGeometryMaskWorldSubsystem>())
 	{
-		Canvas = Subsystem->GetNamedCanvas(GetComponentLevel(), InCanvasName);
+		Canvas = Subsystem->GetNamedCanvas(GetComponentLevel(), InCanvasName, this);
 
 		CanvasWeak = Canvas;
 		
@@ -225,3 +225,9 @@ bool UGeometryMaskCanvasReferenceComponentBase::Cleanup()
 
 	return true;
 }
+
+bool UGeometryMaskCanvasReferenceComponentBase::ForEachUsedCanvasName(TFunctionRef<bool(FName)> InFunc) const
+{
+	return true;
+}
+
diff --git a/Engine/Plugins/VirtualProduction/GeometryMask/Source/GeometryMask/Private/GeometryMaskWorldSubsystem.cpp b/Engine/Plugins/VirtualProduction/GeometryMask/Source/GeometryMask/Private/GeometryMaskWorldSubsystem.cpp
--- a/Engine/Plugins/VirtualProduction/GeometryMask/Source/GeometryMask/Private/GeometryMaskWorldSubsystem.cpp
+++ b/Engine/Plugins/VirtualProduction/GeometryMask/Source/GeometryMask/Private/GeometryMaskWorldSubsystem.cpp
@@ -9,6 +9,7 @@
 #include "GeometryMaskCanvasSharedData.h"
 #include "GeometryMaskCanvasUtils.h"
 #include "GeometryMaskSVE.h"
+#include "IGeometryMaskClient.h"
 #include "SceneViewExtension.h"
 
 void UGeometryMaskWorldSubsystem::Initialize(FSubsystemCollectionBase& Collection)
@@ -74,47 +75,128 @@ FGeometryMaskLevelState::FGeometryMaskLevelState()
 {
 }
 
-int32 FGeometryMaskLevelState::AcquireNextSliceIndex()
+int32 FGeometryMaskLevelState::AcquireAvailableSliceIndex()
 {
-	if (RenderTarget)
+	if (!RenderTarget)
 	{
-		const int32 SliceIndex = NextSliceIndex++;
-		if (SliceIndex >= RenderTarget->Slices)
+		return INDEX_NONE;
+	}
+
+	CleanupUnusedCanvases();
+
+	// amount of slices that should never be reached for a single level.
+	constexpr int32 MaxSlices = 128;
+
+	for (int32 SliceIndex = 0; SliceIndex < MaxSlices; ++SliceIndex)
+	{
+		bool bSliceIndexAvailable = true;
+
+		for (const TWeakInterfacePtr<const IGeometryMaskClient>& MaskClientWeak : MaskClients)
+		{
+			if (const IGeometryMaskClient* MaskClient = MaskClientWeak.Get())
+			{
+				MaskClient->ForEachUsedCanvasName(
+					[This=this, SliceIndex, &bSliceIndexAvailable](FName InCanvasName)
+					{
+						const UGeometryMaskCanvas* const Canvas = This->NamedCanvases.FindRef(InCanvasName);
+						if (Canvas && SliceIndex == Canvas->GetRenderTargetSliceIndex())
+						{
+							bSliceIndexAvailable = false;
+							return false; // break
+						}
+						return true; // continue
+					});
+			}
+		}
+
+		if (bSliceIndexAvailable)
+		{
+			// Update slices and render target if slice index requires a larger number of slices to be present
+			if (SliceIndex >= RenderTarget->Slices)
+			{
+				RenderTarget->Slices = SliceIndex + 1;
+				UE::GeometryMask::UpdateRenderTarget(RenderTarget);
+			}
+			return SliceIndex;
+		}
+	}
+
+	ensureMsgf(0, TEXT("Unexpectedly found no index available below %d"), MaxSlices);
+	return INDEX_NONE;
+}
+
+void FGeometryMaskLevelState::CleanupUnusedCanvases()
+{
+	TSet<FName> UsedCanvasNames;
+	for (const TWeakInterfacePtr<const IGeometryMaskClient>& MaskClientWeak : MaskClients)
+	{
+		if (const IGeometryMaskClient* MaskClient = MaskClientWeak.Get())
+		{
+			MaskClient->ForEachUsedCanvasName(
+				[&UsedCanvasNames](FName InCanvasName)
+				{
+					UsedCanvasNames.Add(InCanvasName);
+					return true; // continue
+				});
+		}
+	}
+
+	for (TMap<FName, TObjectPtr<UGeometryMaskCanvas>>::TIterator Iter(NamedCanvases); Iter; ++Iter)
+	{
+		if (!UsedCanvasNames.Contains(Iter->Key))
 		{
-			// Update RT size to viewport size
-			RenderTarget->Slices = SliceIndex + 1;
-			UE::GeometryMask::UpdateRenderTarget(RenderTarget);
+			if (Iter->Value)
+			{
+				Iter->Value->Free();
+			}
+			Iter.RemoveCurrent();
 		}
-		return SliceIndex;
 	}
-	return -1;
 }
 
 UGeometryMaskCanvas* UGeometryMaskWorldSubsystem::GetNamedCanvas(const ULevel* InLevel, FName InName)
+{
+	return GetNamedCanvas(InLevel, InName, /*MaskClient*/nullptr);
+}
+
+UGeometryMaskCanvas* UGeometryMaskWorldSubsystem::GetNamedCanvas(const ULevel* InLevel, FName InName, const IGeometryMaskClient* InMaskClient)
 {
 	if (InName.IsNone() || !InLevel)
 	{
 		return nullptr;
 	}
 
-	if (const FGeometryMaskLevelState* LevelState = FindLevelState(InLevel))
+	// Registers the mask client to the level state
+	auto RegisterMaskClient = [InMaskClient](FGeometryMaskLevelState& InLevelState)
+		{
+			if (InMaskClient)
+			{
+				InLevelState.MaskClients.AddUnique(InMaskClient);
+				InLevelState.CleanupUnusedCanvases();
+			}
+		};
+
+	if (FGeometryMaskLevelState* LevelState = LevelStates.Find(InLevel))
 	{
-		if (const TObjectPtr<UGeometryMaskCanvas>* FoundCanvas = LevelState->NamedCanvases.Find(InName))
+		if (UGeometryMaskCanvas* FoundCanvas = LevelState->NamedCanvases.FindRef(InName))
 		{
-			return *FoundCanvas;
+			RegisterMaskClient(*LevelState);
+			return FoundCanvas;
 		}
 	}
 
 	const FName ObjectName = MakeUniqueObjectName(this, UGeometryMaskCanvas::StaticClass(), FName(FString::Printf(TEXT("GeometryMaskCanvas_%s_"), *InName.ToString())));
 
 	FGeometryMaskLevelState& LevelState = FindOrAddLevelState(InLevel);
+	RegisterMaskClient(LevelState);
+
 	const TObjectPtr<UGeometryMaskCanvas>& NewCanvas = LevelState.NamedCanvases.Emplace(InName, NewObject<UGeometryMaskCanvas>(this, ObjectName));
 
 	UGeometryMaskCanvas::FInitParams InitParams;
 	InitParams.Level = InLevel;
 	InitParams.CanvasName = InName;
 	InitParams.RenderTarget = LevelState.RenderTarget;
-	InitParams.SliceIndex = LevelState.AcquireNextSliceIndex();
+	InitParams.SliceIndex = LevelState.AcquireAvailableSliceIndex();
 	InitParams.SharedData = LevelState.GetSharedData();
 
 	NewCanvas->Initialize(InitParams);
diff --git a/Engine/Plugins/VirtualProduction/GeometryMask/Source/GeometryMask/Private/GeometryMaskWriteComponent.cpp b/Engine/Plugins/VirtualProduction/GeometryMask/Source/GeometryMask/Private/GeometryMaskWriteComponent.cpp
--- a/Engine/Plugins/VirtualProduction/GeometryMask/Source/GeometryMask/Private/GeometryMaskWriteComponent.cpp
+++ b/Engine/Plugins/VirtualProduction/GeometryMask/Source/GeometryMask/Private/GeometryMaskWriteComponent.cpp
@@ -151,6 +151,11 @@ void UGeometryMaskWriteMeshComponent::OnComponentDestroyed(bool bDestroyingHiera
 	Cleanup();
 }
 
+bool UGeometryMaskWriteMeshComponent::ForEachUsedCanvasName(TFunctionRef<bool(FName)> InFunc) const
+{
+	return InFunc(Parameters.CanvasName);
+}
+
 void UGeometryMaskWriteMeshComponent::UpdateCachedData()
 {
 	if (const AActor* Owner = GetOwner())
diff --git a/Engine/Plugins/VirtualProduction/GeometryMask/Source/GeometryMask/Public/GeometryMaskCanvasActor.h b/Engine/Plugins/VirtualProduction/GeometryMask/Source/GeometryMask/Public/GeometryMaskCanvasActor.h
--- a/Engine/Plugins/VirtualProduction/GeometryMask/Source/GeometryMask/Public/GeometryMaskCanvasActor.h
+++ b/Engine/Plugins/VirtualProduction/GeometryMask/Source/GeometryMask/Public/GeometryMaskCanvasActor.h
@@ -3,7 +3,7 @@
 #pragma once
 
 #include "GameFramework/Actor.h"
-
+#include "IGeometryMaskClient.h"
 #include "GeometryMaskCanvasActor.generated.h"
 
 class IGeometryMaskWriteInterface;
@@ -13,8 +13,7 @@ class UTextureRenderTarget2DArray;
 
 /** Wraps a GeometryMaskCanvas, and discovers/registers writers. */
 UCLASS(BlueprintType)
-class GEOMETRYMASK_API AGeometryMaskCanvasActor
-	: public AActor
+class GEOMETRYMASK_API AGeometryMaskCanvasActor : public AActor, public IGeometryMaskClient
 {
 	GENERATED_BODY()
 
@@ -48,6 +47,10 @@ protected:
 	/** Find writers on child actors and add them to the referenced canvas. */
 	void FindWriters();
 
+	//~ Begin IGeometryMaskClient
+	virtual bool ForEachUsedCanvasName(TFunctionRef<bool(FName)> InFunc) const override;
+	//~ End IGeometryMaskClient
+
 	/** Reference to the Canvas used, identified by CanvasName. */
 	UPROPERTY(Transient, DuplicateTransient)
 	TObjectPtr<UGeometryMaskCanvas> Canvas;
diff --git a/Engine/Plugins/VirtualProduction/GeometryMask/Source/GeometryMask/Public/GeometryMaskReadComponent.h b/Engine/Plugins/VirtualProduction/GeometryMask/Source/GeometryMask/Public/GeometryMaskReadComponent.h
--- a/Engine/Plugins/VirtualProduction/GeometryMask/Source/GeometryMask/Public/GeometryMaskReadComponent.h
+++ b/Engine/Plugins/VirtualProduction/GeometryMask/Source/GeometryMask/Public/GeometryMaskReadComponent.h
@@ -32,6 +32,10 @@ public:
 	//~ End IGeometryMaskReadInterface
 
 protected:
+	//~ Begin IGeometryMaskClient
+	virtual bool ForEachUsedCanvasName(TFunctionRef<bool(FName)> InFunc) const override;
+	//~ End IGeometryMaskClient
+
 	//~ Begin UGeometryMaskCanvasReferenceComponentBase
 	virtual bool TryResolveCanvas() override;
 	//~ End UGeometryMaskCanvasReferenceComponentBase
diff --git a/Engine/Plugins/VirtualProduction/GeometryMask/Source/GeometryMask/Public/GeometryMaskTypes.h b/Engine/Plugins/VirtualProduction/GeometryMask/Source/GeometryMask/Public/GeometryMaskTypes.h
--- a/Engine/Plugins/VirtualProduction/GeometryMask/Source/GeometryMask/Public/GeometryMaskTypes.h
+++ b/Engine/Plugins/VirtualProduction/GeometryMask/Source/GeometryMask/Public/GeometryMaskTypes.h
@@ -4,6 +4,7 @@
 
 #include "Components/ActorComponent.h"
 #include "CoreTypes.h"
+#include "IGeometryMaskClient.h"
 #include "UObject/ObjectKey.h"
 
 #include "GeometryMaskTypes.generated.h"
@@ -236,8 +237,7 @@ struct FGeometryMaskWriteParameters
 };
 
 UCLASS(Abstract)
-class GEOMETRYMASK_API UGeometryMaskCanvasReferenceComponentBase
-	: public UActorComponent
+class GEOMETRYMASK_API UGeometryMaskCanvasReferenceComponentBase : public UActorComponent, public IGeometryMaskClient
 {
 	GENERATED_BODY()
 
@@ -267,8 +267,11 @@ protected:
 	bool TryResolveNamedCanvas(FName InCanvasName);
 
 	[[maybe_unused]] virtual bool Cleanup();
-	
-protected:
+
+	//~ Begin IGeometryMaskClient
+	virtual bool ForEachUsedCanvasName(TFunctionRef<bool(FName)> InFunc) const override;
+	//~ End IGeometryMaskClient
+
 	FOnGeometryMaskSetCanvasNativeDelegate OnSetCanvasDelegate;
 
 	/** Reference to the Canvas used, identified by CanvasName. */
diff --git a/Engine/Plugins/VirtualProduction/GeometryMask/Source/GeometryMask/Public/GeometryMaskWorldSubsystem.h b/Engine/Plugins/VirtualProduction/GeometryMask/Source/GeometryMask/Public/GeometryMaskWorldSubsystem.h
--- a/Engine/Plugins/VirtualProduction/GeometryMask/Source/GeometryMask/Public/GeometryMaskWorldSubsystem.h
+++ b/Engine/Plugins/VirtualProduction/GeometryMask/Source/GeometryMask/Public/GeometryMaskWorldSubsystem.h
@@ -11,6 +11,7 @@
 #define UE_API GEOMETRYMASK_API
 
 class FGeometryMaskSceneViewExtension;
+class IGeometryMaskClient;
 class ULevel;
 struct FGeometryMaskCanvasSharedData;
 
@@ -28,6 +29,9 @@ struct FGeometryMaskLevelState
 	UPROPERTY()
 	TMap<FName, TObjectPtr<UGeometryMaskCanvas>> NamedCanvases;
 
+	/** All the clients using the mask canvases */
+	TArray<TWeakInterfacePtr<const IGeometryMaskClient>> MaskClients;
+
 	/** The render target where the canvases in the levels draw to */
 	UPROPERTY()
 	TObjectPtr<UTextureRenderTarget2DArray> RenderTarget;
@@ -37,14 +41,13 @@ struct FGeometryMaskLevelState
 		return SharedData;
 	}
 
-	/** Returns the index to the next available slice in the render target */
-	int32 AcquireNextSliceIndex();
+	/** Returns the index to the first available slice in the render target */
+	int32 AcquireAvailableSliceIndex();
 
-private:
-	/** tracker for what the next available slice index is */
-	UPROPERTY()
-	int32 NextSliceIndex = 0;
+	/** Cleans up all canvases that were registered but are no longer being used by any of the registered mask clients */
+	void CleanupUnusedCanvases();
 
+private:
 	/** Shared data for canvases that use the render target */
 	TSharedRef<FGeometryMaskCanvasSharedData> SharedData;
 };
@@ -59,6 +62,14 @@ public:
 	/** Retrieves a Canvas, uniquely identified by this world and the canvas name. */
 	UE_API UGeometryMaskCanvas* GetNamedCanvas(const ULevel* InLevel, FName InName);
 
+	/** 
+	 * Retrieves a Canvas, uniquely identified by this world and the canvas name.
+	 * @param InLevel the level that the canvas relates to
+	 * @param InName the name of the canvas
+	 * @param InMaskClient the user of the mask (if any)
+	 */
+	UE_API UGeometryMaskCanvas* GetNamedCanvas(const ULevel* InLevel, FName InName, const IGeometryMaskClient* InMaskClient);
+
 	/** Returns all registered canvas names for this world. */
 	UE_API TArray<FName> GetCanvasNames(const ULevel* InLevel);
 	
diff --git a/Engine/Plugins/VirtualProduction/GeometryMask/Source/GeometryMask/Public/GeometryMaskWriteComponent.h b/Engine/Plugins/VirtualProduction/GeometryMask/Source/GeometryMask/Public/GeometryMaskWriteComponent.h
--- a/Engine/Plugins/VirtualProduction/GeometryMask/Source/GeometryMask/Public/GeometryMaskWriteComponent.h
+++ b/Engine/Plugins/VirtualProduction/GeometryMask/Source/GeometryMask/Public/GeometryMaskWriteComponent.h
@@ -69,6 +69,10 @@ protected:
 	virtual void OnComponentDestroyed(bool bDestroyingHierarchy) override;
 	//~ End UActorComponent
 	
+	//~ Begin IGeometryMaskClient
+	virtual bool ForEachUsedCanvasName(TFunctionRef<bool(FName)> InFunc) const override;
+	//~ End IGeometryMaskClient
+
 	//~ Begin UGeometryMaskCanvasReferenceComponentBase
     virtual bool TryResolveCanvas() override;
 	virtual bool Cleanup() override;
diff --git a/Engine/Plugins/VirtualProduction/GeometryMask/Source/GeometryMask/Public/IGeometryMaskClient.h b/Engine/Plugins/VirtualProduction/GeometryMask/Source/GeometryMask/Public/IGeometryMaskClient.h
new file mode 100644
--- /dev/null
+++ b/Engine/Plugins/VirtualProduction/GeometryMask/Source/GeometryMask/Public/IGeometryMaskClient.h
@@ -0,0 +1,24 @@
+// Copyright Epic Games, Inc. All Rights Reserved.
+
+#pragma once
+
+#include "UObject/Interface.h"
+#include "IGeometryMaskClient.generated.h"
+
+class UGeometryMaskCanvas;
+
+UINTERFACE(MinimalAPI, meta=(CannotImplementInterfaceInBlueprint))
+class UGeometryMaskClient : public UInterface
+{
+	GENERATED_BODY()
+};
+
+/** Interface for objects using canvases */
+class IGeometryMaskClient
+{
+	GENERATED_BODY()
+
+public:
+	/** Iterates each canvas that the client uses */
+	virtual bool ForEachUsedCanvasName(TFunctionRef<bool(FName)> InFunc) const = 0;
+};
