Motion Design: fixed issue where level reload would break input masks for dynamic meshes like shapes.

#Synced-CL 51290016

diff --git a/Engine/Plugins/VirtualProduction/GeometryMask/Source/GeometryMask/Public/GeometryMaskWriteComponent.h b/Engine/Plugins/VirtualProduction/GeometryMask/Source/GeometryMask/Public/GeometryMaskWriteComponent.h
--- a/Engine/Plugins/VirtualProduction/GeometryMask/Source/GeometryMask/Public/GeometryMaskWriteComponent.h
+++ b/Engine/Plugins/VirtualProduction/GeometryMask/Source/GeometryMask/Public/GeometryMaskWriteComponent.h
@@ -105,10 +105,10 @@ protected:
 	TMap<FName, TWeakObjectPtr<UPrimitiveComponent>> CachedComponentsWeak;
 
 	/** Cached primitive component to mesh name, multiple components can use same mesh data */
-	UPROPERTY(DuplicateTransient)
+	UPROPERTY(Transient, DuplicateTransient)
 	TMap<TWeakObjectPtr<UPrimitiveComponent>, FName> CachedComponents;
 
 	/** Cached mesh name to data (one entry per unique mesh) */
-	UPROPERTY(DuplicateTransient)
+	UPROPERTY(Transient, DuplicateTransient)
 	TMap<FName, FGeometryMaskBatchElementData> CachedMeshData;
 };
