Motion Design: remove the need to load materials synchronously and not needed by the gizmo component.

#Synced-CL 51441789

diff --git a/Engine/Plugins/VirtualProduction/Avalanche/Source/Avalanche/Private/Framework/AvaGizmoComponent.cpp b/Engine/Plugins/VirtualProduction/Avalanche/Source/Avalanche/Private/Framework/AvaGizmoComponent.cpp
--- a/Engine/Plugins/VirtualProduction/Avalanche/Source/Avalanche/Private/Framework/AvaGizmoComponent.cpp
+++ b/Engine/Plugins/VirtualProduction/Avalanche/Source/Avalanche/Private/Framework/AvaGizmoComponent.cpp
@@ -329,36 +329,25 @@ void UAvaGizmoComponent::SetMaterial(UMaterialInterface* InMaterial)
 	if (Material != InMaterial)
 	{
 		Material = InMaterial;
-		if (Material)
-		{
-			ApplyToEachPrimitiveComponent(
+		ApplyToEachPrimitiveComponent(
 			[Material = Material](UPrimitiveComponent* InComponent)
+			{
+				for (int32 MaterialIdx = 0; MaterialIdx < InComponent->GetNumMaterials(); ++MaterialIdx)
 				{
-					for (int32 MaterialIdx = 0; MaterialIdx < InComponent->GetNumMaterials(); ++MaterialIdx)
-					{
-						InComponent->SetMaterial(MaterialIdx, Material);
-					}
-				});
-		}
+					InComponent->SetMaterial(MaterialIdx, Material);
+				}
+			});
 	}
 }
 
 void UAvaGizmoComponent::SetMaterialToDefault()
 {
-	static constexpr const TCHAR* MaterialPath = TEXT("Material'/Engine/EngineMaterials/DefaultMaterial.DefaultMaterial'");
-	if (UMaterialInterface* FoundMaterial = LoadObject<UMaterialInterface>(GetTransientPackage(), MaterialPath))
-	{
-		SetMaterial(FoundMaterial);
-	}
+	SetMaterial(nullptr);
 }
 
 void UAvaGizmoComponent::SetMaterialToSeeThrough()
 {
-	static constexpr const TCHAR* MaterialPath = TEXT("Material'/Avalanche/M_SeeThrough.M_SeeThrough'");
-	if (UMaterialInterface* FoundMaterial = LoadObject<UMaterialInterface>(GetTransientPackage(), MaterialPath))
-	{
-		SetMaterial(FoundMaterial);
-	}
+	SetMaterial(nullptr);
 }
 
 bool UAvaGizmoComponent::CastsShadow() const
diff --git a/Engine/Plugins/VirtualProduction/Avalanche/Source/Avalanche/Public/Framework/AvaGizmoComponent.h b/Engine/Plugins/VirtualProduction/Avalanche/Source/Avalanche/Public/Framework/AvaGizmoComponent.h
--- a/Engine/Plugins/VirtualProduction/Avalanche/Source/Avalanche/Public/Framework/AvaGizmoComponent.h
+++ b/Engine/Plugins/VirtualProduction/Avalanche/Source/Avalanche/Public/Framework/AvaGizmoComponent.h
@@ -130,6 +130,7 @@ public:
 	void SetMaterialToDefault();
 
 	/** Set the material to a see through one. */
+	UE_DEPRECATED(5.8, "SetMaterialToSeeThrough deprecated. Unused.")
 	void SetMaterialToSeeThrough();
 
 	/** Whether the primitive component should cast a shadow or not. */
diff --git a/Engine/Plugins/VirtualProduction/GeometryMask/Source/GeometryMask/Private/GeometryMaskCanvasResource.cpp b/Engine/Plugins/VirtualProduction/GeometryMask/Source/GeometryMask/Private/GeometryMaskCanvasResource.cpp
--- a/Engine/Plugins/VirtualProduction/GeometryMask/Source/GeometryMask/Private/GeometryMaskCanvasResource.cpp
+++ b/Engine/Plugins/VirtualProduction/GeometryMask/Source/GeometryMask/Private/GeometryMaskCanvasResource.cpp
@@ -120,7 +120,7 @@ const TSharedPtr<FGeometryMaskCanvasSharedData>& UGeometryMaskCanvasResource::Ge
 
 bool UGeometryMaskCanvasResource::UpdateViewportSize()
 {
-	if (!RenderTarget || !SharedData.IsValid())
+	if (!SharedData.IsValid())
 	{
 		return false;
 	}
@@ -162,7 +162,7 @@ bool UGeometryMaskCanvasResource::UpdateViewportSize()
 
 	SharedData->TextureSize = FIntPoint(SizeX, SizeY);
 
-	if (RenderTarget->SizeX == SizeX && RenderTarget->SizeY == SizeY)
+	if (!RenderTarget || (RenderTarget->SizeX == SizeX && RenderTarget->SizeY == SizeY))
 	{
 		return false;
 	}
@@ -341,7 +341,6 @@ void UGeometryMaskCanvasResource::Draw(const ULevel* InLevel, FSceneView& InView
 		// Update viewport size
 		if (UpdateViewportSize())
 		{
-			FlushRenderingCommands();
 			Canvas->SetRenderTarget_GameThread(RenderTarget->GameThread_GetRenderTargetResource());
 		}
 
diff --git a/Engine/Plugins/VirtualProduction/GeometryMask/Source/GeometryMask/Private/GeometryMaskWorldSubsystem.cpp b/Engine/Plugins/VirtualProduction/GeometryMask/Source/GeometryMask/Private/GeometryMaskWorldSubsystem.cpp
--- a/Engine/Plugins/VirtualProduction/GeometryMask/Source/GeometryMask/Private/GeometryMaskWorldSubsystem.cpp
+++ b/Engine/Plugins/VirtualProduction/GeometryMask/Source/GeometryMask/Private/GeometryMaskWorldSubsystem.cpp
@@ -62,19 +62,12 @@ FGeometryMaskLevelState& UGeometryMaskWorldSubsystem::FindOrAddLevelState(const
 		LevelState->RenderTarget->bCanCreateUAV = true;
 		LevelState->RenderTarget->bTargetArraySlicesIndependently = true;
 
-		constexpr int32 Width = 1920;
-		constexpr int32 Height = 1080;
 		constexpr int32 Slices = 1;
-		LevelState->RenderTarget->Init(Width, Height, Slices, PF_R16F);
+		LevelState->RenderTarget->Init(FMath::Max(1, SharedData->TextureSize.X), FMath::Max(1, SharedData->TextureSize.Y), Slices, PF_R16F);
 	}
 	return *LevelState;
 }
 
-FGeometryMaskLevelState::FGeometryMaskLevelState()
-	: SharedData(MakeShared<FGeometryMaskCanvasSharedData>())
-{
-}
-
 int32 FGeometryMaskLevelState::AcquireAvailableSliceIndex()
 {
 	if (!RenderTarget)
@@ -154,6 +147,11 @@ void FGeometryMaskLevelState::CleanupUnusedCanvases()
 	}
 }
 
+UGeometryMaskWorldSubsystem::UGeometryMaskWorldSubsystem()
+	: SharedData(MakeShared<FGeometryMaskCanvasSharedData>())
+{
+}
+
 UGeometryMaskCanvas* UGeometryMaskWorldSubsystem::GetNamedCanvas(const ULevel* InLevel, FName InName)
 {
 	return GetNamedCanvas(InLevel, InName, /*MaskClient*/nullptr);
@@ -197,7 +195,7 @@ UGeometryMaskCanvas* UGeometryMaskWorldSubsystem::GetNamedCanvas(const ULevel* I
 	InitParams.CanvasName = InName;
 	InitParams.RenderTarget = LevelState.RenderTarget;
 	InitParams.SliceIndex = LevelState.AcquireAvailableSliceIndex();
-	InitParams.SharedData = LevelState.GetSharedData();
+	InitParams.SharedData = SharedData;
 
 	NewCanvas->Initialize(InitParams);
 	OnGeometryMaskCanvasCreatedDelegate.Broadcast(NewCanvas);
diff --git a/Engine/Plugins/VirtualProduction/GeometryMask/Source/GeometryMask/Public/GeometryMaskWorldSubsystem.h b/Engine/Plugins/VirtualProduction/GeometryMask/Source/GeometryMask/Public/GeometryMaskWorldSubsystem.h
--- a/Engine/Plugins/VirtualProduction/GeometryMask/Source/GeometryMask/Public/GeometryMaskWorldSubsystem.h
+++ b/Engine/Plugins/VirtualProduction/GeometryMask/Source/GeometryMask/Public/GeometryMaskWorldSubsystem.h
@@ -23,8 +23,6 @@ struct FGeometryMaskLevelState
 {
 	GENERATED_BODY()
 
-	UE_API FGeometryMaskLevelState();
-
 	/** All the canvases in the level */
 	UPROPERTY()
 	TMap<FName, TObjectPtr<UGeometryMaskCanvas>> NamedCanvases;
@@ -36,20 +34,11 @@ struct FGeometryMaskLevelState
 	UPROPERTY()
 	TObjectPtr<UTextureRenderTarget2DArray> RenderTarget;
 
-	const TSharedRef<FGeometryMaskCanvasSharedData>& GetSharedData() const
-	{
-		return SharedData;
-	}
-
 	/** Returns the index to the first available slice in the render target */
 	int32 AcquireAvailableSliceIndex();
 
 	/** Cleans up all canvases that were registered but are no longer being used by any of the registered mask clients */
 	void CleanupUnusedCanvases();
-
-private:
-	/** Shared data for canvases that use the render target */
-	TSharedRef<FGeometryMaskCanvasSharedData> SharedData;
 };
 
 /** Updates the canvases. */
@@ -59,6 +48,8 @@ class UGeometryMaskWorldSubsystem : public UWorldSubsystem
 	GENERATED_BODY()
 
 public:
+	UGeometryMaskWorldSubsystem();
+
 	/** Retrieves a Canvas, uniquely identified by this world and the canvas name. */
 	UE_API UGeometryMaskCanvas* GetNamedCanvas(const ULevel* InLevel, FName InName);
 
@@ -108,6 +99,9 @@ private:
 
 	UPROPERTY()
 	TMap<TWeakObjectPtr<const ULevel>, FGeometryMaskLevelState> LevelStates;
+
+	/** Shared data for canvases of the world presented */
+	TSharedRef<FGeometryMaskCanvasSharedData> SharedData;
 };
 
 #undef UE_API
