Fix non-unity errors

#Synced-CL 51441789

diff --git a/Engine/Plugins/VirtualProduction/GeometryMask/Source/GeometryMask/Private/GeometryMaskWorldSubsystem.cpp b/Engine/Plugins/VirtualProduction/GeometryMask/Source/GeometryMask/Private/GeometryMaskWorldSubsystem.cpp
--- a/Engine/Plugins/VirtualProduction/GeometryMask/Source/GeometryMask/Private/GeometryMaskWorldSubsystem.cpp
+++ b/Engine/Plugins/VirtualProduction/GeometryMask/Source/GeometryMask/Private/GeometryMaskWorldSubsystem.cpp
@@ -148,8 +148,8 @@ void FGeometryMaskLevelState::CleanupUnusedCanvases()
 }
 
 UGeometryMaskWorldSubsystem::UGeometryMaskWorldSubsystem()
-	: SharedData(MakeShared<FGeometryMaskCanvasSharedData>())
 {
+	SharedData = MakeShared<FGeometryMaskCanvasSharedData>();
 }
 
 UGeometryMaskCanvas* UGeometryMaskWorldSubsystem::GetNamedCanvas(const ULevel* InLevel, FName InName)
diff --git a/Engine/Plugins/VirtualProduction/GeometryMask/Source/GeometryMask/Public/GeometryMaskWorldSubsystem.h b/Engine/Plugins/VirtualProduction/GeometryMask/Source/GeometryMask/Public/GeometryMaskWorldSubsystem.h
--- a/Engine/Plugins/VirtualProduction/GeometryMask/Source/GeometryMask/Public/GeometryMaskWorldSubsystem.h
+++ b/Engine/Plugins/VirtualProduction/GeometryMask/Source/GeometryMask/Public/GeometryMaskWorldSubsystem.h
@@ -4,8 +4,8 @@
 
 #include "GeometryMaskSubsystem.h"
 #include "Subsystems/WorldSubsystem.h"
+#include "Templates/SharedPointer.h"
 #include "UnrealClient.h"
-
 #include "GeometryMaskWorldSubsystem.generated.h"
 
 #define UE_API GEOMETRYMASK_API
@@ -101,7 +101,7 @@ private:
 	TMap<TWeakObjectPtr<const ULevel>, FGeometryMaskLevelState> LevelStates;
 
 	/** Shared data for canvases of the world presented */
-	TSharedRef<FGeometryMaskCanvasSharedData> SharedData;
+	TSharedPtr<FGeometryMaskCanvasSharedData> SharedData;
 };
 
 #undef UE_API
