Motion Design: added geometry mask writer to decouple masking logic from actor component. Mask Write modifier no longer needs to create any component.

#Synced-CL 51764840

diff --git a/Engine/Plugins/VirtualProduction/Avalanche/Source/AvalancheMask/Private/Mask2D/AvaMask2DBaseModifier.cpp b/Engine/Plugins/VirtualProduction/Avalanche/Source/AvalancheMask/Private/Mask2D/AvaMask2DBaseModifier.cpp
--- a/Engine/Plugins/VirtualProduction/Avalanche/Source/AvalancheMask/Private/Mask2D/AvaMask2DBaseModifier.cpp
+++ b/Engine/Plugins/VirtualProduction/Avalanche/Source/AvalancheMask/Private/Mask2D/AvaMask2DBaseModifier.cpp
@@ -172,7 +172,7 @@ void UAvaMask2DBaseModifier::UpdateCanvas()
 {
 	// Reset/invalidate cached
 	LastResolvedCanvasName = NAME_None;
-	CanvasWeak.Reset();
+	ResetCanvas();
 
 #if WITH_EDITOR
 	if (!Channel.IsNone())
@@ -563,11 +563,6 @@ PRAGMA_DISABLE_DEPRECATION_WARNINGS
 	MaterialCollectionHandleData_DEPRECATED.Remove(InActor);
 	MaterialCollectionHandles.Remove(InActor);
 PRAGMA_ENABLE_DEPRECATION_WARNINGS
-
-	if (UAvaGizmoComponent* FoundComponent = InActor->FindComponentByClass<UAvaGizmoComponent>())
-	{
-		FoundComponent->DestroyComponent();
-	}
 }
 
 void UAvaMask2DBaseModifier::OnChannelChanged()
@@ -632,6 +627,7 @@ void UAvaMask2DBaseModifier::TryResolveCanvas()
 		{
 			LastResolvedCanvasName = Canvas->GetCanvasName();
 			CanvasWeak = Canvas;
+			OnCanvasSet(Canvas);
 		}
 	}
 }
@@ -670,7 +666,7 @@ UGeometryMaskCanvas* UAvaMask2DBaseModifier::GetCurrentCanvas()
 	if (LastResolvedCanvasName.IsNone() || LastResolvedCanvasName != GetChannel())
 	{
 		LastResolvedCanvasName = NAME_None;
-		CanvasWeak.Reset();
+		ResetCanvas();
 	}
 
 	if (!CanvasWeak.IsValid())
@@ -727,4 +723,10 @@ PRAGMA_DISABLE_DEPRECATION_WARNINGS
 PRAGMA_ENABLE_DEPRECATION_WARNINGS
 }
 
+void UAvaMask2DBaseModifier::ResetCanvas()
+{
+	OnCanvasReset();
+	CanvasWeak.Reset();
+}
+
 #undef LOCTEXT_NAMESPACE
diff --git a/Engine/Plugins/VirtualProduction/Avalanche/Source/AvalancheMask/Private/Mask2D/AvaMask2DWriteModifier.cpp b/Engine/Plugins/VirtualProduction/Avalanche/Source/AvalancheMask/Private/Mask2D/AvaMask2DWriteModifier.cpp
--- a/Engine/Plugins/VirtualProduction/Avalanche/Source/AvalancheMask/Private/Mask2D/AvaMask2DWriteModifier.cpp
+++ b/Engine/Plugins/VirtualProduction/Avalanche/Source/AvalancheMask/Private/Mask2D/AvaMask2DWriteModifier.cpp
@@ -6,9 +6,9 @@
 #include "Components/PrimitiveComponent.h"
 #include "Engine/CanvasRenderTarget2D.h"
 #include "Engine/Engine.h"
-#include "Framework/AvaGizmoComponent.h"
 #include "GeometryMaskCanvas.h"
 #include "GeometryMaskWriteComponent.h"
+#include "GeometryMaskWriter.h"
 #include "Materials/Material.h"
 #include "Materials/MaterialInstanceDynamic.h"
 #include "Modifiers/ActorModifierCoreStack.h"
@@ -23,6 +23,11 @@ const TAvaPropertyChangeDispatcher<UAvaMask2DWriteModifier> UAvaMask2DWriteModif
 };
 #endif
 
+UAvaMask2DWriteModifier::UAvaMask2DWriteModifier()
+{
+	MaskWriter = UE::GeometryMask::FMaskWriter::Create();
+}
+
 #if WITH_EDITOR
 void UAvaMask2DWriteModifier::PostEditChangeProperty(FPropertyChangedEvent& InPropertyChangedEvent)
 {
@@ -35,6 +40,7 @@ void UAvaMask2DWriteModifier::PostEditChangeProperty(FPropertyChangedEvent& InPr
 void UAvaMask2DWriteModifier::Apply()
 {
 	Super::Apply();
+	MaskActors.Reset();
 
 	if (UGeometryMaskCanvas* Canvas = GetCurrentCanvas())
 	{
@@ -46,7 +52,7 @@ void UAvaMask2DWriteModifier::Apply()
 				// If first modifier in the hierarchy is not this modifier, then skip actor
 				if (MaskModifier == This)
 				{
-					This->ApplyWrite(InActor, Canvas);
+					This->ApplyWrite(InActor);
 				}
 				return true;
 			}
@@ -75,34 +81,13 @@ void UAvaMask2DWriteModifier::OnRenderStateUpdated(AActor* InActor, UActorCompon
 	}
 }
 
-bool UAvaMask2DWriteModifier::ApplyWrite(TNotNull<AActor*> InActor, TNotNull<UGeometryMaskCanvas*> InCanvas)
+bool UAvaMask2DWriteModifier::ApplyWrite(TNotNull<AActor*> InActor)
 {
 	// Only add read/write component to actors with primitives
 	if (ActorSupportsMaskReadWrite(InActor))
 	{
-		if (UGeometryMaskWriteMeshComponent* WriteComponent = FindOrAddMaskComponent<UGeometryMaskWriteMeshComponent>(InActor))
-		{
-			FGeometryMaskWriteParameters Parameters = WriteComponent->GetParameters();
-			Parameters.CanvasName = Channel;
-			Parameters.OperationType = WriteOperation;
-			WriteComponent->SetParameters(Parameters);
-		}
-	}
-
-	if (UAvaGizmoComponent* VisualizationComponent = UE::AvaMask::Internal::FindOrAddComponent<UAvaGizmoComponent>(InActor))
-	{
-		VisualizationComponent->SetVisibleInEditor(true);
-		VisualizationComponent->SetHiddenInGame(false);
-		VisualizationComponent->SetShowWireframe(false);
-		VisualizationComponent->SetWireframeColor(FLinearColor::Red);
-		VisualizationComponent->SetsStencil(true);
-		VisualizationComponent->SetStencilId(150);
-		VisualizationComponent->SetRenderInMainPass(false);
-		VisualizationComponent->SetRenderDepth(false);
-		VisualizationComponent->SetMaterialToDefault();
-		VisualizationComponent->ApplyGizmoValues();
+		MaskActors.Add(InActor);
 	}
-
 	return true;
 }
 
@@ -112,7 +97,7 @@ bool UAvaMask2DWriteModifier::ApplyWrite(AActor* InActor, FAvaMask2DActorData& I
 	{
 		if (UGeometryMaskCanvas* Canvas = GetCurrentCanvas())
 		{
-			return ApplyWrite(InActor, Canvas);
+			return ApplyWrite(InActor);
 		}
 	}
 	return true;
@@ -141,6 +126,18 @@ void UAvaMask2DWriteModifier::OnModifierCDOSetup(FActorModifierCoreMetadata& InM
 #endif
 }
 
+void UAvaMask2DWriteModifier::OnModifierEnabled(EActorModifierCoreEnableReason InReason)
+{
+	Super::OnModifierEnabled(InReason);
+	TryResolveCanvas();
+}
+
+void UAvaMask2DWriteModifier::OnModifierDisabled(EActorModifierCoreDisableReason InReason)
+{
+	Super::OnModifierDisabled(InReason);
+	ResetCanvas();
+}
+
 void UAvaMask2DWriteModifier::SetupMaskComponent(UActorComponent* InComponent)
 {
 	if (!InComponent)
@@ -148,14 +145,89 @@ void UAvaMask2DWriteModifier::SetupMaskComponent(UActorComponent* InComponent)
 		return;
 	}
 	
-	if (IGeometryMaskWriteInterface* MaskWriter = Cast<IGeometryMaskWriteInterface>(InComponent))
+	if (IGeometryMaskWriteInterface* Writer = Cast<IGeometryMaskWriteInterface>(InComponent))
 	{
-		SetupMaskWriteComponent(MaskWriter);
+		SetupMaskWriteComponent(Writer);
 	}
 
 	Super::SetupMaskComponent(InComponent);
 }
 
+void UAvaMask2DWriteModifier::OnCanvasSet(TNotNull<UGeometryMaskCanvas*> InCanvas)
+{
+	InCanvas->AddWriter(this);
+}
+
+void UAvaMask2DWriteModifier::OnCanvasReset()
+{
+	if (UGeometryMaskCanvas* Canvas = CanvasWeak.Get())
+	{
+		Canvas->RemoveWriter(this);
+	}
+}
+
+void UAvaMask2DWriteModifier::RemoveFromActor(AActor* InActor)
+{
+	Super::RemoveFromActor(InActor);
+	MaskActors.Remove(InActor);
+}
+
+const FGeometryMaskWriteParameters& UAvaMask2DWriteModifier::GetParameters() const
+{
+	return Parameters;
+}
+
+void UAvaMask2DWriteModifier::SetParameters(FGeometryMaskWriteParameters& InParameters)
+{
+	Parameters = InParameters;
+
+	UGeometryMaskCanvas* Canvas = CanvasWeak.Get();
+
+	// We have changed canvas, unregister this writer
+	if (Canvas && Canvas->GetFName() != InParameters.CanvasName)
+	{
+		ResetCanvas();
+	}
+
+	// Update canvas
+	if (!CanvasWeak.IsValid())
+	{
+		TryResolveCanvas();
+	}
+}
+
+void UAvaMask2DWriteModifier::DrawToCanvas(FCanvas* InCanvas)
+{
+	if (IsModifierEnabled())
+	{
+		Parameters.CanvasName = Channel;
+		Parameters.OperationType = WriteOperation;
+
+		UE::GeometryMask::FMaskWriter::FDrawParams DrawParams
+			{
+				.Canvas = InCanvas,
+				.Parameters = &Parameters,
+				.bWriteWhenHidden = true,
+			};
+
+		for (AActor* Actor : MaskActors)
+		{
+			DrawParams.Actor = Actor;
+			MaskWriter->DrawToCanvas(DrawParams);
+		}
+	}
+}
+
+FOnGeometryMaskSetCanvasNativeDelegate& UAvaMask2DWriteModifier::OnSetCanvas()
+{
+	return OnSetCanvasDelegate;
+}
+
+UE::GeometryMask::FMaskWriter* UAvaMask2DWriteModifier::GetMaskWriter()
+{
+	return MaskWriter.Get();
+}
+
 void UAvaMask2DWriteModifier::SetupMaskWriteComponent(IGeometryMaskWriteInterface* InMaskWriter)
 {
 	if (const UGeometryMaskCanvas* Canvas = GetCurrentCanvas())
diff --git a/Engine/Plugins/VirtualProduction/Avalanche/Source/AvalancheMask/Public/Mask2D/AvaMask2DBaseModifier.h b/Engine/Plugins/VirtualProduction/Avalanche/Source/AvalancheMask/Public/Mask2D/AvaMask2DBaseModifier.h
--- a/Engine/Plugins/VirtualProduction/Avalanche/Source/AvalancheMask/Public/Mask2D/AvaMask2DBaseModifier.h
+++ b/Engine/Plugins/VirtualProduction/Avalanche/Source/AvalancheMask/Public/Mask2D/AvaMask2DBaseModifier.h
@@ -208,6 +208,19 @@ protected:
 	UE_DEPRECATED(5.8, "ActorData has been deprecated. Actor restore takes place in the Material Read Modifier directly.")
 	UE_API void RestoreActors();
 
+	/** Called to reset the canvas. Calls OnCanvasReset */
+	UE_API void ResetCanvas();
+
+	/** Called when the canvas has been set */
+	virtual void OnCanvasSet(TNotNull<UGeometryMaskCanvas*> InCanvas)
+	{
+	}
+
+	/** Called when the canvas has reset*/
+	virtual void OnCanvasReset()
+	{
+	}
+
 #if WITH_EDITOR
 	UFUNCTION(CallInEditor, Category = "Mask2D")
 	UE_API void VisualizeMask();
diff --git a/Engine/Plugins/VirtualProduction/Avalanche/Source/AvalancheMask/Public/Mask2D/AvaMask2DReadModifier.h b/Engine/Plugins/VirtualProduction/Avalanche/Source/AvalancheMask/Public/Mask2D/AvaMask2DReadModifier.h
--- a/Engine/Plugins/VirtualProduction/Avalanche/Source/AvalancheMask/Public/Mask2D/AvaMask2DReadModifier.h
+++ b/Engine/Plugins/VirtualProduction/Avalanche/Source/AvalancheMask/Public/Mask2D/AvaMask2DReadModifier.h
@@ -106,7 +106,7 @@ protected:
 	void StoreActorState(AActor* InActor);
 
 	/** Additional channels to read from */
-	UPROPERTY(EditInstanceOnly, Category = "Mask2D", meta = (EditCondition = "!bUseParentChannel", EditConditionHides, AllowPrivateAccess = "true", DisplayAfter="Channel"))
+	UPROPERTY(EditInstanceOnly, Category = "Mask2D", meta = (AllowPrivateAccess = "true", DisplayAfter="Channel"))
 	TArray<FName> AdditionalChannels;
 
 	/** Base opacity/alpha to use in Read mode */
diff --git a/Engine/Plugins/VirtualProduction/Avalanche/Source/AvalancheMask/Public/Mask2D/AvaMask2DWriteModifier.h b/Engine/Plugins/VirtualProduction/Avalanche/Source/AvalancheMask/Public/Mask2D/AvaMask2DWriteModifier.h
--- a/Engine/Plugins/VirtualProduction/Avalanche/Source/AvalancheMask/Public/Mask2D/AvaMask2DWriteModifier.h
+++ b/Engine/Plugins/VirtualProduction/Avalanche/Source/AvalancheMask/Public/Mask2D/AvaMask2DWriteModifier.h
@@ -5,6 +5,7 @@
 #include "AvaMask2DBaseModifier.h"
 #include "AvaPropertyChangeDispatcher.h"
 #include "GeometryMaskTypes.h"
+#include "IGeometryMaskWriteInterface.h"
 #include "UObject/ObjectPtr.h"
 #include "UObject/Package.h"
 #include "UObject/WeakObjectPtr.h"
@@ -32,15 +33,22 @@ class UTexture;
 struct FAvaShapeParametricMaterial;
 struct FMaterialParameterInfo;
 
+namespace UE::GeometryMask
+{
+	struct FMaskWriter;
+}
+
 /** Uses scene actors to create a mask texture and applies it to attached actors */
 UCLASS(BlueprintType)
-class  UAvaMask2DWriteModifier : public UAvaMask2DBaseModifier
+class UAvaMask2DWriteModifier : public UAvaMask2DBaseModifier, public IGeometryMaskWriteInterface
 {
 	GENERATED_BODY()
 
 	friend class UAvaMask2DReadModifierShared;
 	
 public:
+	UE_API UAvaMask2DWriteModifier();
+
 	EGeometryMaskCompositeOperation GetWriteOperation() const
 	{
 		return WriteOperation;
@@ -56,26 +64,51 @@ protected:
 
 	//~ Begin UActorModifierCoreBase
 	UE_API virtual void OnModifierCDOSetup(FActorModifierCoreMetadata& InMetadata) override;
+	UE_API virtual void OnModifierEnabled(EActorModifierCoreEnableReason InReason) override;
+	UE_API virtual void OnModifierDisabled(EActorModifierCoreDisableReason InReason) override;
 	UE_API virtual void Apply() override;
 	//~ End UActorModifierCoreBase
 
 	//~ Begin UAvaMask2DBaseModifier
 	UE_API virtual void OnRenderStateUpdated(AActor* InActor, UActorComponent* InComponent) override;
 	UE_API virtual void SetupMaskComponent(UActorComponent* InComponent) override;
+	UE_API virtual void OnCanvasSet(TNotNull<UGeometryMaskCanvas*> InCanvas) override;
+	UE_API virtual void OnCanvasReset() override;
+	UE_API virtual void RemoveFromActor(AActor* InActor) override;
 	//~ End UAvaMask2DBaseModifier
 
+	//~ Begin IGeometryMaskWriteInterface
+	UE_API virtual const FGeometryMaskWriteParameters& GetParameters() const override;
+	UE_API virtual void SetParameters(FGeometryMaskWriteParameters& InParameters) override;
+	UE_API virtual void DrawToCanvas(FCanvas* InCanvas) override;
+	UE_API virtual FOnGeometryMaskSetCanvasNativeDelegate& OnSetCanvas() override;
+	UE_API virtual UE::GeometryMask::FMaskWriter* GetMaskWriter() override;
+	//~ End IGeometryMaskWriteInterface
+
 	UE_API void OnWriteOperationChanged();
 	UE_API void SetupMaskWriteComponent(IGeometryMaskWriteInterface* InMaskWriter);
 
-	UE_API bool ApplyWrite(TNotNull<AActor*> InActor, TNotNull<UGeometryMaskCanvas*> InCanvas);
+	UE_API bool ApplyWrite(TNotNull<AActor*> InActor);
 
-	UE_DEPRECATED(5.8, "Actor data has been deprecated. Use the overload that takes in a canvas instead")
+	UE_DEPRECATED(5.8, "Actor data has been deprecated. Use the overload that takes in an actor only")
 	UE_API bool ApplyWrite(AActor* InActor, FAvaMask2DActorData& InActorData);
 
 	/** How to write to the chosen mask channel */
 	UPROPERTY(EditInstanceOnly, BlueprintReadWrite, Getter, Setter, Category = "Mask2D", meta = (ValidEnumValues = "Add, Subtract", AllowPrivateAccess = "true"))
 	EGeometryMaskCompositeOperation WriteOperation = EGeometryMaskCompositeOperation::Add;
 
+	/** Actors to get the components to mask with */
+	UPROPERTY(Transient, DuplicateTransient)
+	TArray<TObjectPtr<AActor>> MaskActors;
+
+	/** Object handling masking and component caching */
+	TSharedPtr<UE::GeometryMask::FMaskWriter> MaskWriter;
+
+	UPROPERTY(Transient, DuplicateTransient)
+	FGeometryMaskWriteParameters Parameters;
+
+	FOnGeometryMaskSetCanvasNativeDelegate OnSetCanvasDelegate;
+
 #if WITH_EDITOR
 	/** Used for PECP */
 	UE_API static const TAvaPropertyChangeDispatcher<UAvaMask2DWriteModifier> PropertyChangeDispatcher;
diff --git a/Engine/Plugins/VirtualProduction/GeometryMask/Source/GeometryMask/Private/GeometryMaskCanvas.cpp b/Engine/Plugins/VirtualProduction/GeometryMask/Source/GeometryMask/Private/GeometryMaskCanvas.cpp
--- a/Engine/Plugins/VirtualProduction/GeometryMask/Source/GeometryMask/Private/GeometryMaskCanvas.cpp
+++ b/Engine/Plugins/VirtualProduction/GeometryMask/Source/GeometryMask/Private/GeometryMaskCanvas.cpp
@@ -2,12 +2,14 @@
 
 #include "GeometryMaskCanvas.h"
 #include "Algo/RemoveIf.h"
+#include "Components/PrimitiveComponent.h"
 #include "Engine/CanvasRenderTarget2D.h"
 #include "Engine/Engine.h"
 #include "Engine/World.h"
 #include "GeometryMaskCanvasResource.h"
 #include "GeometryMaskCanvasSharedData.h"
 #include "GeometryMaskModule.h"
+#include "GeometryMaskWriter.h"
 #include "IGeometryMaskWriteInterface.h"
 
 const FName UGeometryMaskCanvas::ApplyBlurPropertyName = GET_MEMBER_NAME_CHECKED(UGeometryMaskCanvas, bApplyBlur);
@@ -238,7 +240,7 @@ void UGeometryMaskCanvas::RemoveInvalidWriters()
 	}
 }
 
-void UGeometryMaskCanvas::OnDrawToCanvas(const FGeometryMaskDrawingContext& InDrawingContext, FCanvas* InCanvas)
+void UGeometryMaskCanvas::OnDrawToCanvas(const FGeometryMaskDrawingContext& InDrawingContext, FSceneView& InView, FCanvas* InCanvas)
 {
 	if (!ensure(InCanvas))
 	{
@@ -250,13 +252,28 @@ void UGeometryMaskCanvas::OnDrawToCanvas(const FGeometryMaskDrawingContext& InDr
 	// Sort so subtract happens after additive, etc.
 	SortWriters();
 
-	for (const TWeakInterfacePtr<IGeometryMaskWriteInterface>& Writer : Writers)
+	for (const TWeakInterfacePtr<IGeometryMaskWriteInterface>& WriterWeak : Writers)
 	{
-		if (Writer.IsValid())
+		IGeometryMaskWriteInterface* const Writer = WriterWeak.Get();
+		if (!Writer)
 		{
-			TRACE_CPUPROFILER_EVENT_SCOPE(UGeometryMaskCanvas::OnDrawToCanvas_Writer);
-			DECLARE_SCOPE_CYCLE_COUNTER(TEXT("UGeometryMaskCanvas::OnDrawToCanvas"), STAT_GeometryMask_Update, STATGROUP_GeometryMask);
-			Writer->DrawToCanvas(InCanvas);	
+			continue;
+		}
+
+		// Draw the mask primitives to canvas
+		{
+			TRACE_CPUPROFILER_EVENT_SCOPE(IGeometryMaskWriteInterface::OnDrawToCanvas);
+			Writer->DrawToCanvas(InCanvas);
+		}
+
+		// Add the mask primitives to the hidden list
+		if (UE::GeometryMask::FMaskWriter* MaskWriter = Writer->GetMaskWriter())
+		{
+			MaskWriter->ForEachMaskPrimitive(
+				[&InView](TNotNull<const UPrimitiveComponent*> InComponent)
+				{
+					InView.HiddenPrimitives.Add(InComponent->GetPrimitiveSceneId());
+				});
 		}
 	}
 }
@@ -279,7 +296,6 @@ void UGeometryMaskCanvas::Initialize(const FInitParams& InInitParams)
 	CanvasId = FGeometryMaskCanvasId(InInitParams.Level, CanvasName);
 
 	CanvasResource = NewObject<UGeometryMaskCanvasResource>(this);
-	CanvasResource->OnDrawToCanvas().AddUObject(this, &UGeometryMaskCanvas::OnDrawToCanvas);
 	CanvasResource->Initialize(InInitParams.RenderTarget, InInitParams.SliceIndex, InInitParams.SharedData);
 	UpdateRenderParameters();
 }
diff --git a/Engine/Plugins/VirtualProduction/GeometryMask/Source/GeometryMask/Private/GeometryMaskCanvasResource.cpp b/Engine/Plugins/VirtualProduction/GeometryMask/Source/GeometryMask/Private/GeometryMaskCanvasResource.cpp
--- a/Engine/Plugins/VirtualProduction/GeometryMask/Source/GeometryMask/Private/GeometryMaskCanvasResource.cpp
+++ b/Engine/Plugins/VirtualProduction/GeometryMask/Source/GeometryMask/Private/GeometryMaskCanvasResource.cpp
@@ -289,6 +289,12 @@ int16 UGeometryMaskCanvasResource::GetRenderTargetSliceIndex() const
 	return RenderTargetSliceIndex;
 }
 
+FOnGeometryMaskCanvasDraw& UGeometryMaskCanvasResource::OnDrawToCanvas()
+{
+	static FOnGeometryMaskCanvasDraw OnDrawToCanvasDelegate;
+	return OnDrawToCanvasDelegate;
+}
+
 void UGeometryMaskCanvasResource::Update(const ULevel* InLevel, FSceneView& InView, int32 InViewIndex)
 {
 	Draw(InLevel, InView, InViewIndex);
@@ -307,6 +313,12 @@ void UGeometryMaskCanvasResource::Draw(const ULevel* InLevel, FSceneView& InView
 		return;
 	}
 
+	UGeometryMaskCanvas* const Owner = Cast<UGeometryMaskCanvas>(GetOuter());
+	if (!Owner)
+	{
+		return;
+	}
+
 	TRACE_CPUPROFILER_EVENT_SCOPE(UGeometryMaskCanvasResource::Draw);
 
 	UE_LOG(LogGeometryMask, VeryVerbose, TEXT("UGeometryMaskCanvasResource::Update Level: %s"), *InLevel->GetName());
@@ -364,10 +376,8 @@ void UGeometryMaskCanvasResource::Draw(const ULevel* InLevel, FSceneView& InView
 
 		// Set to World->Viewport transform
 		Canvas->SetBaseTransform(DrawingContext->ViewProjectionMatrix);
-		{
-			TRACE_CPUPROFILER_EVENT_SCOPE(UGeometryMaskCanvasResource::Draw::OnDrawToCanvas);
-			OnDrawToCanvas().Broadcast(*DrawingContext, Canvas.Get());
-		}
+
+		Owner->OnDrawToCanvas(*DrawingContext, InView, Canvas.Get());
 
 		// Restore original transform
 		Canvas->SetBaseTransform(CanvasMatrix);
diff --git a/Engine/Plugins/VirtualProduction/GeometryMask/Source/GeometryMask/Private/GeometryMaskWriteComponent.cpp b/Engine/Plugins/VirtualProduction/GeometryMask/Source/GeometryMask/Private/GeometryMaskWriteComponent.cpp
--- a/Engine/Plugins/VirtualProduction/GeometryMask/Source/GeometryMask/Private/GeometryMaskWriteComponent.cpp
+++ b/Engine/Plugins/VirtualProduction/GeometryMask/Source/GeometryMask/Private/GeometryMaskWriteComponent.cpp
@@ -1,7 +1,6 @@
 ﻿// Copyright Epic Games, Inc. All Rights Reserved.
 
 #include "GeometryMaskWriteComponent.h"
-
 #include "Algo/RemoveIf.h"
 #include "BatchedElements.h"
 #include "CanvasTypes.h"
@@ -10,9 +9,15 @@
 #include "Engine/CanvasRenderTarget2D.h"
 #include "Engine/StaticMesh.h"
 #include "GeometryMaskCanvas.h"
+#include "GeometryMaskWriter.h"
 #include "GlobalRenderResources.h"
 #include "StaticMeshResources.h"
 
+UGeometryMaskWriteMeshComponent::UGeometryMaskWriteMeshComponent()
+{
+	MaskWriter = UE::GeometryMask::FMaskWriter::Create();
+}
+
 void UGeometryMaskWriteMeshComponent::SetParameters(FGeometryMaskWriteParameters& InParameters)
 {
 	Parameters = InParameters;
@@ -35,113 +40,18 @@ void UGeometryMaskWriteMeshComponent::SetParameters(FGeometryMaskWriteParameters
 
 void UGeometryMaskWriteMeshComponent::DrawToCanvas(FCanvas* InCanvas)
 {
-	TRACE_CPUPROFILER_EVENT_SCOPE(UGeometryMaskWriteMeshComponent::DrawToCanvas);
-
-	if (!bWriteWhenHidden)
-	{
-		if (const AActor* Owner = GetOwner())
+	MaskWriter->DrawToCanvas(UE::GeometryMask::FMaskWriter::FDrawParams
 		{
-			if (Owner->IsHidden())
-			{
-				return;
-			}
-		}
-	}
-
-	UpdateCachedData();
-
-	const FLinearColor Color = GetParameters().OperationType == EGeometryMaskCompositeOperation::Subtract 
-		? FLinearColor::Black 
-		: FLinearColor::White;
-
-	constexpr ESimpleElementBlendMode ElementBlendMode = ESimpleElementBlendMode::SE_BLEND_Opaque;
-
-	// Write to Canvas BatchElements
-	{
-		FHitProxyId CanvasHitProxyId = InCanvas->GetHitProxyId();
-
-		TMap<TWeakObjectPtr<UPrimitiveComponent>, FName> ComponentsToKeep;
-		ComponentsToKeep.Reserve(CachedComponents.Num());
-
-		TRACE_CPUPROFILER_EVENT_SCOPE(UGeometryMaskWriteMeshComponent::DrawToCanvas::WriteToCanvas);
-
-		// Flag valid components, can't do async due to UObject access
-		for (const TPair<TWeakObjectPtr<UPrimitiveComponent>, FName>& CachedComponentWeak : CachedComponents)
-		{
-			if (UPrimitiveComponent* CachedComponent = CachedComponentWeak.Key.Get())
-			{
-				if (CachedMeshData.Contains(CachedComponentWeak.Value))
-				{
-					ComponentsToKeep.Emplace(CachedComponent, CachedComponentWeak.Value);
-				}
-			}
-		}
-
-		for (const TPair<TWeakObjectPtr<UPrimitiveComponent>, FName>& ComponentPair : ComponentsToKeep)
-		{
-			TRACE_CPUPROFILER_EVENT_SCOPE(UGeometryMaskWriteMeshComponent::DrawToCanvas::WriteToCanvas::Task);
-
-			const UPrimitiveComponent* PrimitiveComponent = ComponentPair.Key.Get();
-
-			if (!PrimitiveComponent)
-			{
-				continue;
-			}
-
-			const FTransform LocalToWorld = PrimitiveComponent->GetComponentToWorld();
-			const FGeometryMaskBatchElementData& MeshBatchElementData = CachedMeshData[ComponentPair.Value];
-
-			InCanvas->PushAbsoluteTransform(LocalToWorld.ToMatrixWithScale());
-			{
-				FBatchedElements* CanvasTriangleBatchedElements = InCanvas->GetBatchedElements(FCanvas::ET_Triangle);
-
-				CanvasTriangleBatchedElements->AddReserveVertices(MeshBatchElementData.Vertices.Num());
-				CanvasTriangleBatchedElements->AddReserveTriangles(MeshBatchElementData.NumTriangles, GWhiteTexture, ElementBlendMode);
-
-				auto AddVertex = [CanvasTriangleBatchedElements, &MeshBatchElementData, CanvasHitProxyId, &Color](int32 VertexIdx) -> int32
-				{
-					return CanvasTriangleBatchedElements->AddVertexf(
-						MeshBatchElementData.Vertices[VertexIdx],
-						FVector2f::ZeroVector,
-						Color,
-						CanvasHitProxyId);
-				};
-
-				// Store the index of the first vertex to be able to convert the MeshBatchElementData vertex indices to
-				// map to the actual vertices added to the BatchedElements.
-				int32 InitialIndex = 0;
-
-				if (MeshBatchElementData.Vertices.Num() > 0)
-				{
-					TRACE_CPUPROFILER_EVENT_SCOPE(UGeometryMaskWriteMeshComponent::DrawToCanvas::WriteToCanvas::AddVertices);
-
-					InitialIndex = AddVertex(0);
-					for (int32 VertexIdx = 1; VertexIdx < MeshBatchElementData.Vertices.Num(); ++VertexIdx)
-					{
-						AddVertex(VertexIdx);
-					}
-				}
-
-				{
-					TRACE_CPUPROFILER_EVENT_SCOPE(UGeometryMaskWriteMeshComponent::DrawToCanvas::WriteToCanvas::AddTriangles);
-					for (int32 VertexIdx = 0; VertexIdx <= MeshBatchElementData.Indices.Num() - 3; VertexIdx += 3)
-					{
-						// MeshBatchElementData.Indices are 'relative' in the sense that do not consider possible existing 
-						// vertices within BatchedElements prior to its vertices being added.
-						// So these 'local' indices are converted to be the proper indices of the vertices that were added above. 
-						const int32 V0 = InitialIndex + MeshBatchElementData.Indices[VertexIdx];
-						const int32 V1 = InitialIndex + MeshBatchElementData.Indices[VertexIdx + 1];
-						const int32 V2 = InitialIndex + MeshBatchElementData.Indices[VertexIdx + 2];
-
-						CanvasTriangleBatchedElements->AddTriangle(V0, V1, V2, GWhiteTexture, ElementBlendMode);
-					}
-				}
-			}
-			InCanvas->PopTransform();
-		}
+			.Canvas = InCanvas, 
+			.Actor = GetOwner(),
+			.Parameters = &Parameters, 
+			.bWriteWhenHidden = bWriteWhenHidden,
+		});
+}
 
-		CachedComponents = ComponentsToKeep;
-	}
+UE::GeometryMask::FMaskWriter* UGeometryMaskWriteMeshComponent::GetMaskWriter()
+{
+	return MaskWriter.Get();
 }
 
 void UGeometryMaskWriteMeshComponent::OnComponentDestroyed(bool bDestroyingHierarchy)
@@ -158,282 +68,24 @@ bool UGeometryMaskWriteMeshComponent::ForEachUsedCanvasName(TFunctionRef<bool(FN
 
 void UGeometryMaskWriteMeshComponent::UpdateCachedData()
 {
-	if (const AActor* Owner = GetOwner())
-	{
-		TArray<UPrimitiveComponent*> PrimitiveComponents;
-		Owner->GetComponents<UPrimitiveComponent>(PrimitiveComponents);
-
-		// @todo: better check for cached data state?
-		if (LastPrimitiveComponentCount == PrimitiveComponents.Num())
-		{
-			return;
-		}
-
-		LastPrimitiveComponentCount = PrimitiveComponents.Num();
-		CachedComponents.Reserve(LastPrimitiveComponentCount);
-		CachedMeshData.Reserve(LastPrimitiveComponentCount);
-
-		UpdateCachedStaticMeshData(PrimitiveComponents);
-		UpdateCachedDynamicMeshData(PrimitiveComponents);
-	}
 }
 
 void UGeometryMaskWriteMeshComponent::UpdateCachedStaticMeshData(TConstArrayView<UPrimitiveComponent*> InPrimitiveComponents)
 {
-	TArray<UStaticMeshComponent*> StaticMeshComponents;
-	StaticMeshComponents.Reserve(InPrimitiveComponents.Num());
-	
-	Algo::TransformIf(
-		InPrimitiveComponents,
-		StaticMeshComponents,
-		[](const UPrimitiveComponent* InComponent)
-		{
-			return InComponent->IsA<UStaticMeshComponent>();
-		},
-		[](UPrimitiveComponent* InComponent)
-		{
-			return Cast<UStaticMeshComponent>(InComponent);
-		});
-
-	if (StaticMeshComponents.IsEmpty())
-	{
-		return;
-	}
-
-	// Remove built/already cached
-	StaticMeshComponents.SetNum(Algo::RemoveIf(StaticMeshComponents, [this](const UStaticMeshComponent* InStaticMeshComponent)
-	{
-		const UStaticMesh* StaticMesh = InStaticMeshComponent->GetStaticMesh();
-	
-		if (!StaticMesh || !StaticMesh->HasValidRenderData())
-		{
-			return true;
-		}
-
-		const FName ResourceName = StaticMesh->GetFName();
-		if (CachedComponents.Contains(InStaticMeshComponent) && CachedMeshData.Contains(ResourceName))
-		{
-			return true;
-		}
-
-		return false;
-	}));
-
-	// Subscribe to change events
-#if WITH_EDITOR
-	for (UStaticMeshComponent* StaticMeshComponent : StaticMeshComponents)
-	{
-		StaticMeshComponent->OnStaticMeshChanged().RemoveAll(this);
-		StaticMeshComponent->OnStaticMeshChanged().AddUObject(this, &UGeometryMaskWriteMeshComponent::OnStaticMeshChanged);
-	}
-#endif
-
-	// Used for cache lookup
-	TArray<FName> StaticMeshObjectNames;
-	StaticMeshObjectNames.Reserve(StaticMeshComponents.Num());
-	
-	TArray<FStaticMeshLODResources*> StaticMeshResources;
-	StaticMeshResources.Reserve(StaticMeshComponents.Num());
-
-	// Collect valid mesh resources
-	{
-		TRACE_CPUPROFILER_EVENT_SCOPE(UGeometryMaskWriteMeshComponent::DrawToCanvas::CollectMeshResources);
-		for (UStaticMeshComponent* StaticMeshComponent : StaticMeshComponents)
-		{
-			if (UStaticMesh* StaticMesh = StaticMeshComponent->GetStaticMesh())
-			{
-				if (!StaticMesh->HasValidRenderData())
-				{
-					continue;
-				}
-			
-				if (FStaticMeshRenderData* RenderData = StaticMesh->GetRenderData())
-				{
-					if (RenderData->LODResources.IsEmpty())
-					{
-						continue;
-					}
-
-					if (RenderData->LODResources[0].GetNumVertices() == 0)
-					{
-						continue;
-					}
-
-					const FName ResourceName = StaticMesh->GetFName();
-					StaticMeshObjectNames.Add(ResourceName);
-					StaticMeshResources.Add(&RenderData->LODResources[0]);
-					CachedComponents.Add(StaticMeshComponent, ResourceName);
-				}
-			}
-		}
-
-		// Convert mesh resources to batch elements
-		{
-			TRACE_CPUPROFILER_EVENT_SCOPE(UGeometryMaskWriteMeshComponent::DrawToCanvas::ConvertMeshResources);
-			
-			CachedMeshData.Reserve(CachedMeshData.Num() + StaticMeshObjectNames.Num());
-
-			for (int32 MeshIdx = 0; MeshIdx < StaticMeshResources.Num(); ++MeshIdx)
-			{
-				TRACE_CPUPROFILER_EVENT_SCOPE(UGeometryMaskWriteMeshComponent::DrawToCanvas::ConvertMeshResources::Task);
-
-				FStaticMeshLODResources* MeshResources = StaticMeshResources[MeshIdx];
-				const int32 NumVertices = MeshResources->GetNumVertices();
-				const int32 NumIndices = MeshResources->IndexBuffer.GetNumIndices();
-				const int32 NumTriangles = MeshResources->GetNumTriangles();
-
-				FGeometryMaskBatchElementData& MeshBatchElementData = CachedMeshData.Emplace(StaticMeshObjectNames[MeshIdx]);
-				MeshBatchElementData.Reserve(NumVertices, NumIndices, NumTriangles);
-				MeshResources->IndexBuffer.GetCopy(MeshBatchElementData.Indices);
-
-				for (int32 VertexIdx = 0; VertexIdx < NumVertices; ++VertexIdx)
-				{
-					MeshBatchElementData.Vertices.Add(FVector4f(MeshResources->VertexBuffers.PositionVertexBuffer.VertexPosition(VertexIdx)));
-				}
-			}
-		}
-	}
 }
 
 void UGeometryMaskWriteMeshComponent::UpdateCachedDynamicMeshData(TConstArrayView<UPrimitiveComponent*> InPrimitiveComponents)
 {
-	TArray<UDynamicMeshComponent*> DynamicMeshComponents;
-	DynamicMeshComponents.Reserve(InPrimitiveComponents.Num());
-	
-	Algo::TransformIf(
-		InPrimitiveComponents,
-		DynamicMeshComponents,
-		[](const UPrimitiveComponent* InComponent)
-		{
-			return InComponent->IsA<UDynamicMeshComponent>();
-		},
-		[](UPrimitiveComponent* InComponent)
-		{
-			return Cast<UDynamicMeshComponent>(InComponent);
-		});
-
-	if (DynamicMeshComponents.IsEmpty())
-	{
-		return;
-	}
-
-	// Remove built/already cached
-	DynamicMeshComponents.SetNum(Algo::RemoveIf(DynamicMeshComponents, [this](UDynamicMeshComponent* InDynamicMeshComponent)
-	{
-		const UDynamicMesh* DynamicMesh = InDynamicMeshComponent->GetDynamicMesh();
-	
-		if (!DynamicMesh || !DynamicMesh->GetMeshPtr())
-		{
-			return true;
-		}
-
-		const FName ResourceName = InDynamicMeshComponent->GetFName();
-		if (CachedComponents.Contains(InDynamicMeshComponent) && CachedMeshData.Contains(ResourceName))
-		{
-			return true;
-		}
-
-		return false;
-	}));
-
-	// Subscribe to change events
-	for (UDynamicMeshComponent* DynamicMeshComponent : DynamicMeshComponents)
-	{
-		if (FDynamicMesh3* DynamicMesh3 = DynamicMeshComponent->GetMesh())
-		{
-			DynamicMesh3->SetShapeChangeStampEnabled(true);
-		}
-
-		DynamicMeshComponent->OnMeshChanged.RemoveAll(this);
-		DynamicMeshComponent->OnMeshChanged.AddUObject(this, &UGeometryMaskWriteMeshComponent::OnDynamicMeshChanged, DynamicMeshComponent);
-	}
-
-	// Used for cache lookup
-	TArray<FName> DynamicMeshObjectNames;
-	DynamicMeshObjectNames.Reserve(DynamicMeshComponents.Num());
-	
-	TArray<FDynamicMesh3*> DynamicMeshes;
-	DynamicMeshes.Reserve(DynamicMeshComponents.Num());
-
-	// Collect valid meshes
-	{
-		TRACE_CPUPROFILER_EVENT_SCOPE(UGeometryMaskWriteMeshComponent::DrawToCanvas::CollectMeshResources);
-		for (UDynamicMeshComponent* DynamicMeshComponent : DynamicMeshComponents)
-		{
-			if (UDynamicMesh* DynamicMeshObject = DynamicMeshComponent->GetDynamicMesh())
-			{
-				if (FDynamicMesh3* DynamicMesh = DynamicMeshObject->GetMeshPtr())
-				{
-					if (DynamicMesh->TriangleCount() == 0)
-					{
-						continue;
-					}
-
-					const FName ResourceName = DynamicMeshComponent->GetFName();
-					DynamicMeshObjectNames.Add(ResourceName);
-					DynamicMeshes.Add(DynamicMesh);
-					CachedComponents.Add(DynamicMeshComponent, ResourceName);
-				}
-			}
-		}
-	}
-
-	// Convert meshes to batch elements
-	{
-		TRACE_CPUPROFILER_EVENT_SCOPE(UGeometryMaskWriteMeshComponent::DrawToCanvas::ConvertMeshResources);
-		
-		CachedMeshData.Reserve(CachedMeshData.Num() + DynamicMeshObjectNames.Num());
-		
-		for (int32 MeshIdx = 0; MeshIdx < DynamicMeshes.Num(); ++MeshIdx)
-		{
-			TRACE_CPUPROFILER_EVENT_SCOPE(UGeometryMaskWriteMeshComponent::DrawToCanvas::ConvertMeshResources::Task);
-
-			FDynamicMesh3 CompactMesh;
-			const FDynamicMesh3* DynamicMesh = DynamicMeshes[MeshIdx];
-			CompactMesh.CompactCopy(*DynamicMesh);
-
-			const int32 NumVertices = CompactMesh.VertexCount();
-			const int32 NumIndices = CompactMesh.TriangleCount() * 3;
-			const int32 NumTriangles = CompactMesh.TriangleCount();
-
-			FGeometryMaskBatchElementData& MeshBatchElementData = CachedMeshData.Emplace(DynamicMeshObjectNames[MeshIdx]);
-			MeshBatchElementData.ChangeStamp = DynamicMesh->GetChangeStamp();
-			MeshBatchElementData.Reserve(NumVertices, NumIndices, NumTriangles);
-
-			for (const int32 VertexIdx : CompactMesh.VertexIndicesItr())
-			{
-				MeshBatchElementData.Vertices.Add(FVector4f(FVector3f(CompactMesh.GetVertex(VertexIdx))));				
-			}
-
-			for (const UE::Geometry::FIndex3i Triangle : CompactMesh.TrianglesItr())
-			{
-				MeshBatchElementData.Indices.Append({
-					static_cast<uint32>(Triangle.A),
-					static_cast<uint32>(Triangle.B),
-					static_cast<uint32>(Triangle.C)});
-			}
-		}
-	}
 }
 
 #if WITH_EDITOR
 void UGeometryMaskWriteMeshComponent::OnStaticMeshChanged(UStaticMeshComponent* InStaticMeshComponent)
 {
-	// Triggers a cache refresh
-	ResetCachedData();
 }
 #endif
 
 void UGeometryMaskWriteMeshComponent::OnDynamicMeshChanged(UDynamicMeshComponent* InDynamicMeshComponent)
 {
-	// Triggers a specific cache refresh
-	if (FGeometryMaskBatchElementData* CachedData = CachedMeshData.Find(InDynamicMeshComponent->GetFName()))
-	{
-		CachedData->ChangeStamp = INDEX_NONE;
-	}
-	
-	// Triggers a general cache refresh
-	ResetCachedData();
 }
 
 bool UGeometryMaskWriteMeshComponent::TryResolveCanvas()
@@ -470,7 +122,4 @@ bool UGeometryMaskWriteMeshComponent::Cleanup()
 
 void UGeometryMaskWriteMeshComponent::ResetCachedData()
 {
-	LastPrimitiveComponentCount = -1;
-	CachedComponents.Reset();
-	CachedMeshData.Reset();
 }
diff --git a/Engine/Plugins/VirtualProduction/GeometryMask/Source/GeometryMask/Private/GeometryMaskWriter.cpp b/Engine/Plugins/VirtualProduction/GeometryMask/Source/GeometryMask/Private/GeometryMaskWriter.cpp
new file mode 100644
--- /dev/null
+++ b/Engine/Plugins/VirtualProduction/GeometryMask/Source/GeometryMask/Private/GeometryMaskWriter.cpp
@@ -0,0 +1,404 @@
+// Copyright Epic Games, Inc. All Rights Reserved.
+
+#include "GeometryMaskWriter.h"
+#include "Algo/RemoveIf.h"
+#include "BatchedElements.h"
+#include "CanvasTypes.h"
+#include "Components/DynamicMeshComponent.h"
+#include "Components/PrimitiveComponent.h"
+#include "Engine/StaticMesh.h"
+#include "Engine/World.h"
+#include "GameFramework/Actor.h"
+#include "GlobalRenderResources.h"
+#include "IndexTypes.h"
+#include "StaticMeshResources.h"
+
+namespace UE::GeometryMask
+{
+
+TSharedRef<FMaskWriter> FMaskWriter::Create()
+{
+	return MakeShared<FMaskWriter>();
+}
+
+void FMaskWriter::DrawToCanvas(const FDrawParams& InParams)
+{
+	if (!InParams.Canvas || (!InParams.bWriteWhenHidden && InParams.Actor && InParams.Actor->IsHidden()))
+	{
+		return;
+	}
+
+	TRACE_CPUPROFILER_EVENT_SCOPE(FGeometryMaskWriter::DrawToCanvas);
+	UpdateCachedData(InParams.Actor);
+
+	const FLinearColor Color = InParams.Parameters->OperationType == EGeometryMaskCompositeOperation::Subtract 
+		? FLinearColor::Black 
+		: FLinearColor::White;
+
+	constexpr ESimpleElementBlendMode ElementBlendMode = ESimpleElementBlendMode::SE_BLEND_Opaque;
+	const FHitProxyId CanvasHitProxyId = InParams.Canvas->GetHitProxyId();
+
+	for (TMap<TObjectKey<UPrimitiveComponent>, FObjectKey>::TIterator Iter(CachedComponents); Iter; ++Iter)
+	{
+		if (!CachedMeshData.Contains(Iter->Value))
+		{
+			Iter.RemoveCurrent();
+			continue;
+		}
+
+		const UPrimitiveComponent* const Component = Iter->Key.ResolveObjectPtr();
+		if (!Component)
+		{
+			Iter.RemoveCurrent();
+			continue;
+		}
+
+		const FTransform& LocalToWorld = Component->GetComponentToWorld();
+		const FGeometryMaskBatchElementData& MeshBatchElementData = CachedMeshData[Iter->Value];
+
+		InParams.Canvas->PushAbsoluteTransform(LocalToWorld.ToMatrixWithScale());
+		ON_SCOPE_EXIT
+		{
+			InParams.Canvas->PopTransform();
+		};
+
+		FBatchedElements* const CanvasTriangleBatchedElements = InParams.Canvas->GetBatchedElements(FCanvas::ET_Triangle);
+		CanvasTriangleBatchedElements->AddReserveVertices(MeshBatchElementData.Vertices.Num());
+		CanvasTriangleBatchedElements->AddReserveTriangles(MeshBatchElementData.NumTriangles, GWhiteTexture, ElementBlendMode);
+
+		auto AddVertex = [CanvasTriangleBatchedElements, &MeshBatchElementData, CanvasHitProxyId, &Color](int32 VertexIdx) -> int32
+			{
+				return CanvasTriangleBatchedElements->AddVertexf(
+					MeshBatchElementData.Vertices[VertexIdx],
+					FVector2f::ZeroVector,
+					Color,
+					CanvasHitProxyId);
+			};
+
+		// Store the index of the first vertex to be able to convert the MeshBatchElementData vertex indices to
+		// map to the actual vertices added to the BatchedElements.
+		int32 InitialIndex = 0;
+
+		if (!MeshBatchElementData.Vertices.IsEmpty())
+		{
+			InitialIndex = AddVertex(0);
+			for (int32 VertexIdx = 1; VertexIdx < MeshBatchElementData.Vertices.Num(); ++VertexIdx)
+			{
+				AddVertex(VertexIdx);
+			}
+		}
+
+		for (int32 VertexIdx = 0; VertexIdx <= MeshBatchElementData.Indices.Num() - 3; VertexIdx += 3)
+		{
+			// MeshBatchElementData.Indices are 'relative' in the sense that do not consider possible existing 
+			// vertices within BatchedElements prior to its vertices being added.
+			// So these 'local' indices are converted to be the proper indices of the vertices that were added above. 
+			const int32 V0 = InitialIndex + MeshBatchElementData.Indices[VertexIdx];
+			const int32 V1 = InitialIndex + MeshBatchElementData.Indices[VertexIdx + 1];
+			const int32 V2 = InitialIndex + MeshBatchElementData.Indices[VertexIdx + 2];
+
+			CanvasTriangleBatchedElements->AddTriangle(V0, V1, V2, GWhiteTexture, ElementBlendMode);
+		}
+	}
+}
+
+void FMaskWriter::ForEachMaskPrimitive(TFunctionRef<void(TNotNull<const UPrimitiveComponent*>)> InFunc) const
+{
+	for (const TPair<TObjectKey<UPrimitiveComponent>, FObjectKey>& Pair : CachedComponents)
+	{
+		if (UPrimitiveComponent* Component = Pair.Key.ResolveObjectPtr())
+		{
+			InFunc(Component);
+		}
+	}
+}
+
+void FMaskWriter::ResetCachedData()
+{
+	LastPrimitiveComponentCountMap.Reset();
+	CachedComponents.Reset();
+	CachedMeshData.Reset();
+}
+
+void FMaskWriter::UpdateCachedData(AActor* InActor)
+{
+	if (InActor)
+	{
+		TArray<UPrimitiveComponent*> PrimitiveComponents;
+		InActor->GetComponents<UPrimitiveComponent>(PrimitiveComponents);
+
+		// @todo: better check for cached data state?
+		int32& LastPrimitiveComponentCount = LastPrimitiveComponentCountMap.FindOrAdd(InActor, 0);
+		if (LastPrimitiveComponentCount == PrimitiveComponents.Num())
+		{
+			return;
+		}
+
+		LastPrimitiveComponentCount = PrimitiveComponents.Num();
+		UpdateCachedStaticMeshData(PrimitiveComponents);
+		UpdateCachedDynamicMeshData(PrimitiveComponents);
+	}
+}
+
+void FMaskWriter::UpdateCachedStaticMeshData(TConstArrayView<UPrimitiveComponent*> InPrimitiveComponents)
+{
+	TArray<UStaticMeshComponent*> StaticMeshComponents;
+	StaticMeshComponents.Reserve(InPrimitiveComponents.Num());
+
+	Algo::TransformIf(
+		InPrimitiveComponents,
+		StaticMeshComponents,
+		[](const UPrimitiveComponent* InComponent)
+		{
+			return InComponent->IsA<UStaticMeshComponent>();
+		},
+		[](UPrimitiveComponent* InComponent)
+		{
+			return Cast<UStaticMeshComponent>(InComponent);
+		});
+
+	if (StaticMeshComponents.IsEmpty())
+	{
+		return;
+	}
+
+	// Remove built/already cached
+	StaticMeshComponents.SetNum(Algo::RemoveIf(StaticMeshComponents, [this](const UStaticMeshComponent* InStaticMeshComponent)
+	{
+		const UStaticMesh* StaticMesh = InStaticMeshComponent->GetStaticMesh();
+	
+		if (!StaticMesh || !StaticMesh->HasValidRenderData())
+		{
+			return true;
+		}
+
+		if (CachedComponents.Contains(InStaticMeshComponent) && CachedMeshData.Contains(StaticMesh))
+		{
+			return true;
+		}
+
+		return false;
+	}));
+
+	// Subscribe to change events
+#if WITH_EDITOR
+	for (UStaticMeshComponent* StaticMeshComponent : StaticMeshComponents)
+	{
+		StaticMeshComponent->OnStaticMeshChanged().RemoveAll(this);
+		StaticMeshComponent->OnStaticMeshChanged().AddSPLambda(this, 
+			[this](UStaticMeshComponent* InStaticMeshComponent)
+			{
+				OnStaticMeshChanged(InStaticMeshComponent);
+			});
+	}
+#endif
+
+	// Used for cache lookup
+	TArray<FObjectKey> StaticMeshObjects;
+	StaticMeshObjects.Reserve(StaticMeshComponents.Num());
+
+	TArray<FStaticMeshLODResources*> StaticMeshResources;
+	StaticMeshResources.Reserve(StaticMeshComponents.Num());
+
+	// Collect valid mesh resources
+	{
+		TRACE_CPUPROFILER_EVENT_SCOPE(UGeometryMaskWriteMeshComponent::DrawToCanvas::CollectMeshResources);
+		for (UStaticMeshComponent* StaticMeshComponent : StaticMeshComponents)
+		{
+			if (UStaticMesh* StaticMesh = StaticMeshComponent->GetStaticMesh())
+			{
+				if (!StaticMesh->HasValidRenderData())
+				{
+					continue;
+				}
+			
+				if (FStaticMeshRenderData* RenderData = StaticMesh->GetRenderData())
+				{
+					if (RenderData->LODResources.IsEmpty())
+					{
+						continue;
+					}
+
+					if (RenderData->LODResources[0].GetNumVertices() == 0)
+					{
+						continue;
+					}
+
+					StaticMeshObjects.Add(StaticMesh);
+					StaticMeshResources.Add(&RenderData->LODResources[0]);
+					CachedComponents.Add(StaticMeshComponent, StaticMesh);
+				}
+			}
+		}
+
+		// Convert mesh resources to batch elements
+		{
+			TRACE_CPUPROFILER_EVENT_SCOPE(UGeometryMaskWriteMeshComponent::DrawToCanvas::ConvertMeshResources);
+			
+			CachedMeshData.Reserve(CachedMeshData.Num() + StaticMeshObjects.Num());
+
+			for (int32 MeshIdx = 0; MeshIdx < StaticMeshResources.Num(); ++MeshIdx)
+			{
+				TRACE_CPUPROFILER_EVENT_SCOPE(UGeometryMaskWriteMeshComponent::DrawToCanvas::ConvertMeshResources::Task);
+
+				FStaticMeshLODResources* MeshResources = StaticMeshResources[MeshIdx];
+				const int32 NumVertices = MeshResources->GetNumVertices();
+				const int32 NumIndices = MeshResources->IndexBuffer.GetNumIndices();
+				const int32 NumTriangles = MeshResources->GetNumTriangles();
+
+				FGeometryMaskBatchElementData& MeshBatchElementData = CachedMeshData.Emplace(StaticMeshObjects[MeshIdx]);
+				MeshBatchElementData.Reserve(NumVertices, NumIndices, NumTriangles);
+				MeshResources->IndexBuffer.GetCopy(MeshBatchElementData.Indices);
+
+				for (int32 VertexIdx = 0; VertexIdx < NumVertices; ++VertexIdx)
+				{
+					MeshBatchElementData.Vertices.Add(FVector4f(MeshResources->VertexBuffers.PositionVertexBuffer.VertexPosition(VertexIdx)));
+				}
+			}
+		}
+	}
+}
+
+void FMaskWriter::UpdateCachedDynamicMeshData(TConstArrayView<UPrimitiveComponent*> InPrimitiveComponents)
+{
+	TArray<UDynamicMeshComponent*> DynamicMeshComponents;
+	DynamicMeshComponents.Reserve(InPrimitiveComponents.Num());
+	
+	Algo::TransformIf(
+		InPrimitiveComponents,
+		DynamicMeshComponents,
+		[](const UPrimitiveComponent* InComponent)
+		{
+			return InComponent->IsA<UDynamicMeshComponent>();
+		},
+		[](UPrimitiveComponent* InComponent)
+		{
+			return Cast<UDynamicMeshComponent>(InComponent);
+		});
+
+	if (DynamicMeshComponents.IsEmpty())
+	{
+		return;
+	}
+
+	// Remove built/already cached
+	DynamicMeshComponents.SetNum(Algo::RemoveIf(DynamicMeshComponents, [this](UDynamicMeshComponent* InDynamicMeshComponent)
+	{
+		const UDynamicMesh* DynamicMesh = InDynamicMeshComponent->GetDynamicMesh();
+	
+		if (!DynamicMesh || !DynamicMesh->GetMeshPtr())
+		{
+			return true;
+		}
+
+		if (CachedComponents.Contains(InDynamicMeshComponent) && CachedMeshData.Contains(DynamicMesh))
+		{
+			return true;
+		}
+
+		return false;
+	}));
+
+	// Subscribe to change events
+	for (UDynamicMeshComponent* DynamicMeshComponent : DynamicMeshComponents)
+	{
+		if (FDynamicMesh3* DynamicMesh3 = DynamicMeshComponent->GetMesh())
+		{
+			DynamicMesh3->SetShapeChangeStampEnabled(true);
+		}
+
+		DynamicMeshComponent->OnMeshChanged.RemoveAll(this);
+		DynamicMeshComponent->OnMeshChanged.AddSPLambda(this, 
+			[this, DynamicMeshComponent]()
+			{
+				OnDynamicMeshChanged(DynamicMeshComponent);
+			});
+	}
+
+	// Used for cache lookup
+	TArray<FObjectKey> DynamicMeshObjects;
+	DynamicMeshObjects.Reserve(DynamicMeshComponents.Num());
+	
+	TArray<FDynamicMesh3*> DynamicMeshes;
+	DynamicMeshes.Reserve(DynamicMeshComponents.Num());
+
+	// Collect valid meshes
+	{
+		TRACE_CPUPROFILER_EVENT_SCOPE(UGeometryMaskWriteMeshComponent::DrawToCanvas::CollectMeshResources);
+		for (UDynamicMeshComponent* DynamicMeshComponent : DynamicMeshComponents)
+		{
+			if (UDynamicMesh* DynamicMeshObject = DynamicMeshComponent->GetDynamicMesh())
+			{
+				if (FDynamicMesh3* DynamicMesh = DynamicMeshObject->GetMeshPtr())
+				{
+					if (DynamicMesh->TriangleCount() == 0)
+					{
+						continue;
+					}
+
+					DynamicMeshObjects.Add(DynamicMeshObject);
+					DynamicMeshes.Add(DynamicMesh);
+					CachedComponents.Add(DynamicMeshComponent, DynamicMeshObject);
+				}
+			}
+		}
+	}
+
+	// Convert meshes to batch elements
+	{
+		TRACE_CPUPROFILER_EVENT_SCOPE(UGeometryMaskWriteMeshComponent::DrawToCanvas::ConvertMeshResources);
+		
+		CachedMeshData.Reserve(CachedMeshData.Num() + DynamicMeshObjects.Num());
+		
+		for (int32 MeshIdx = 0; MeshIdx < DynamicMeshes.Num(); ++MeshIdx)
+		{
+			TRACE_CPUPROFILER_EVENT_SCOPE(UGeometryMaskWriteMeshComponent::DrawToCanvas::ConvertMeshResources::Task);
+
+			FDynamicMesh3 CompactMesh;
+			const FDynamicMesh3* DynamicMesh = DynamicMeshes[MeshIdx];
+			CompactMesh.CompactCopy(*DynamicMesh);
+
+			const int32 NumVertices = CompactMesh.VertexCount();
+			const int32 NumIndices = CompactMesh.TriangleCount() * 3;
+			const int32 NumTriangles = CompactMesh.TriangleCount();
+
+			FGeometryMaskBatchElementData& MeshBatchElementData = CachedMeshData.Emplace(DynamicMeshObjects[MeshIdx]);
+			MeshBatchElementData.ChangeStamp = DynamicMesh->GetChangeStamp();
+			MeshBatchElementData.Reserve(NumVertices, NumIndices, NumTriangles);
+
+			for (const int32 VertexIdx : CompactMesh.VertexIndicesItr())
+			{
+				MeshBatchElementData.Vertices.Add(FVector4f(FVector3f(CompactMesh.GetVertex(VertexIdx))));				
+			}
+
+			for (const UE::Geometry::FIndex3i Triangle : CompactMesh.TrianglesItr())
+			{
+				MeshBatchElementData.Indices.Append({
+					static_cast<uint32>(Triangle.A),
+					static_cast<uint32>(Triangle.B),
+					static_cast<uint32>(Triangle.C)});
+			}
+		}
+	}
+}
+
+#if WITH_EDITOR
+void FMaskWriter::OnStaticMeshChanged(UStaticMeshComponent* InStaticMeshComponent)
+{
+	// Triggers a cache refresh
+	ResetCachedData();
+}
+#endif
+
+void FMaskWriter::OnDynamicMeshChanged(UDynamicMeshComponent* InDynamicMeshComponent)
+{
+	// Triggers a specific cache refresh
+	if (FGeometryMaskBatchElementData* CachedData = CachedMeshData.Find(InDynamicMeshComponent->GetDynamicMesh()))
+	{
+		CachedData->ChangeStamp = INDEX_NONE;
+	}
+	// Triggers a general cache refresh
+	ResetCachedData();
+}
+
+} // UE::GeometryMask
diff --git a/Engine/Plugins/VirtualProduction/GeometryMask/Source/GeometryMask/Public/GeometryMaskCanvas.h b/Engine/Plugins/VirtualProduction/GeometryMask/Source/GeometryMask/Public/GeometryMaskCanvas.h
--- a/Engine/Plugins/VirtualProduction/GeometryMask/Source/GeometryMask/Public/GeometryMaskCanvas.h
+++ b/Engine/Plugins/VirtualProduction/GeometryMask/Source/GeometryMask/Public/GeometryMaskCanvas.h
@@ -184,7 +184,7 @@ private:
 	void RemoveInvalidWriters();
 
 	/** Draws all writers to the canvas. */
-	void OnDrawToCanvas(const FGeometryMaskDrawingContext& InDrawingContext, FCanvas* InCanvas);
+	void OnDrawToCanvas(const FGeometryMaskDrawingContext& InDrawingContext, FSceneView& InView, FCanvas* InCanvas);
 
 	/** Updates relevant shader parameters for the given color channel. */
 	void UpdateRenderParameters();
diff --git a/Engine/Plugins/VirtualProduction/GeometryMask/Source/GeometryMask/Public/GeometryMaskCanvasResource.h b/Engine/Plugins/VirtualProduction/GeometryMask/Source/GeometryMask/Public/GeometryMaskCanvasResource.h
--- a/Engine/Plugins/VirtualProduction/GeometryMask/Source/GeometryMask/Public/GeometryMaskCanvasResource.h
+++ b/Engine/Plugins/VirtualProduction/GeometryMask/Source/GeometryMask/Public/GeometryMaskCanvasResource.h
@@ -97,10 +97,8 @@ public:
 	/** Get the index to the render target slice the resource draws to */
 	UE_API int16 GetRenderTargetSliceIndex() const;
 
-	FOnGeometryMaskCanvasDraw& OnDrawToCanvas()
-	{
-		return OnDrawToCanvasDelegate;
-	}
+	UE_DEPRECATED(5.8, "OnDrawToCanvasDelegate has been deprecated as canvas resources are now directly owned by the canvas object")
+	UE_API FOnGeometryMaskCanvasDraw& OnDrawToCanvas();
 
 	UE_DEPRECATED(5.8, "Color Channel is no longer used. Use the overload that does not take the color channel")
 	UE_API void UpdateRenderParameters(EGeometryMaskColorChannel InColorChannel, bool bInApplyBlur, double InBlurStrength, bool bInApplyFeather, int32 InOuterFeatherRadius, int32 InInnerFeatherRadius);
@@ -165,8 +163,6 @@ private:
 	/** Shared data between canvases using the same render target */
 	TSharedPtr<FGeometryMaskCanvasSharedData> SharedData;
 
-	FOnGeometryMaskCanvasDraw OnDrawToCanvasDelegate;
-
 	TSet<FGeometryMaskDrawingContext> DrawingContextCache;
 
 	FGeometryMaskDrawingContext DefaultDrawingContext;
diff --git a/Engine/Plugins/VirtualProduction/GeometryMask/Source/GeometryMask/Public/GeometryMaskTypes.h b/Engine/Plugins/VirtualProduction/GeometryMask/Source/GeometryMask/Public/GeometryMaskTypes.h
--- a/Engine/Plugins/VirtualProduction/GeometryMask/Source/GeometryMask/Public/GeometryMaskTypes.h
+++ b/Engine/Plugins/VirtualProduction/GeometryMask/Source/GeometryMask/Public/GeometryMaskTypes.h
@@ -278,3 +278,33 @@ protected:
 	UPROPERTY(Transient, DuplicateTransient, TextExportTransient)
 	TWeakObjectPtr<UGeometryMaskCanvas> CanvasWeak;
 };
+
+/** Mesh data for a single mesh, in local space. */
+USTRUCT()
+struct FGeometryMaskBatchElementData
+{
+	GENERATED_BODY()
+
+	UPROPERTY()
+	TArray<FVector3f> Vertices;
+
+	UPROPERTY()
+	TArray<uint32> Indices;
+
+	UPROPERTY()
+	int32 NumTriangles = 0;
+
+	/** Used to check cached state. */
+	UPROPERTY()
+	uint64 ChangeStamp = 0;
+
+	void Reserve(
+		const int32 InNumVertices,
+		const int32 InNumIndices,
+		const int32 InNumTriangles)
+	{
+		Vertices.Reserve(InNumVertices);
+		Indices.Reserve(InNumIndices);
+		NumTriangles = InNumTriangles;
+	}
+};
diff --git a/Engine/Plugins/VirtualProduction/GeometryMask/Source/GeometryMask/Public/GeometryMaskWriteComponent.h b/Engine/Plugins/VirtualProduction/GeometryMask/Source/GeometryMask/Public/GeometryMaskWriteComponent.h
--- a/Engine/Plugins/VirtualProduction/GeometryMask/Source/GeometryMask/Public/GeometryMaskWriteComponent.h
+++ b/Engine/Plugins/VirtualProduction/GeometryMask/Source/GeometryMask/Public/GeometryMaskWriteComponent.h
@@ -5,90 +5,75 @@
 #include "Components/ActorComponent.h"
 #include "GeometryMaskTypes.h"
 #include "IGeometryMaskWriteInterface.h"
-
 #include "GeometryMaskWriteComponent.generated.h"
 
+#define UE_API GEOMETRYMASK_API
+
 class UCanvasRenderTarget2D;
 class UDynamicMeshComponent;
 class UPrimitiveComponent;
 class UStaticMeshComponent;
 
-/** Mesh data for a single mesh, in local space. */
-USTRUCT()
-struct FGeometryMaskBatchElementData
+namespace UE::GeometryMask
 {
-	GENERATED_BODY()
-
-	UPROPERTY()
-	TArray<FVector3f> Vertices;
-
-	UPROPERTY()
-	TArray<uint32> Indices;
-
-	UPROPERTY()
-	int32 NumTriangles = 0;
-
-	/** Used to check cached state. */
-	UPROPERTY()
-	uint64 ChangeStamp = 0;
-
-	void Reserve(
-		const int32 InNumVertices,
-		const int32 InNumIndices,
-		const int32 InNumTriangles)
-	{
-		Vertices.Reserve(InNumVertices);
-		Indices.Reserve(InNumIndices);
-		NumTriangles = InNumTriangles;
-	}
-};
+	struct FMaskWriter;
+}
 
-UCLASS(BlueprintType, HideCategories=(Activation, Cooking, AssetUserData, Navigation), meta=(BlueprintSpawnableComponent))
-class GEOMETRYMASK_API UGeometryMaskWriteMeshComponent
-	: public UGeometryMaskCanvasReferenceComponentBase
-	, public IGeometryMaskWriteInterface
+UCLASS(MinimalAPI, BlueprintType, HideCategories=(Activation, Cooking, AssetUserData, Navigation), meta=(BlueprintSpawnableComponent))
+class UGeometryMaskWriteMeshComponent: public UGeometryMaskCanvasReferenceComponentBase, public IGeometryMaskWriteInterface
 {
 	GENERATED_BODY()
 
 public:
+	UE_API UGeometryMaskWriteMeshComponent();
+
 	//~ Begin IGeometryMaskWriteInterface
-	virtual void SetParameters(FGeometryMaskWriteParameters& InParameters) override;
+	UE_API virtual void SetParameters(FGeometryMaskWriteParameters& InParameters) override;
 	virtual const FGeometryMaskWriteParameters& GetParameters() const override
 	{
 		return Parameters;
 	}
-	virtual void DrawToCanvas(FCanvas* InCanvas) override;
+	UE_API virtual void DrawToCanvas(FCanvas* InCanvas) override;
 	virtual FOnGeometryMaskSetCanvasNativeDelegate& OnSetCanvas() override
 	{
 		return OnSetCanvasDelegate;
 	}
+	UE_API virtual UE::GeometryMask::FMaskWriter* GetMaskWriter() override;
 	//~ End IGeometryMaskWriteInterface
 
 protected:
 	//~ Begin UActorComponent
-	virtual void OnComponentDestroyed(bool bDestroyingHierarchy) override;
+	UE_API virtual void OnComponentDestroyed(bool bDestroyingHierarchy) override;
 	//~ End UActorComponent
-	
+
 	//~ Begin IGeometryMaskClient
-	virtual bool ForEachUsedCanvasName(TFunctionRef<bool(FName)> InFunc) const override;
+	UE_API virtual bool ForEachUsedCanvasName(TFunctionRef<bool(FName)> InFunc) const override;
 	//~ End IGeometryMaskClient
 
 	//~ Begin UGeometryMaskCanvasReferenceComponentBase
-    virtual bool TryResolveCanvas() override;
-	virtual bool Cleanup() override;
+    UE_API virtual bool TryResolveCanvas() override;
+	UE_API virtual bool Cleanup() override;
     //~ End UGeometryMaskCanvasReferenceComponentBase
 
 	// Resets cached data, triggers rebuild
-	void ResetCachedData();
-	
-	void UpdateCachedData();
-	void UpdateCachedStaticMeshData(TConstArrayView<UPrimitiveComponent*> InPrimitiveComponents);
-	void UpdateCachedDynamicMeshData(TConstArrayView<UPrimitiveComponent*> InPrimitiveComponents);
+	UE_DEPRECATED(5.8, "Logic moved to UE::GeometryMask::FMaskWriter")
+	UE_API void ResetCachedData();
+
+	UE_DEPRECATED(5.8, "Logic moved to UE::GeometryMask::FMaskWriter")
+	UE_API void UpdateCachedData();
+
+	UE_DEPRECATED(5.8, "Logic moved to UE::GeometryMask::FMaskWriter")
+	UE_API void UpdateCachedStaticMeshData(TConstArrayView<UPrimitiveComponent*> InPrimitiveComponents);
+
+	UE_DEPRECATED(5.8, "Logic moved to UE::GeometryMask::FMaskWriter")
+	UE_API void UpdateCachedDynamicMeshData(TConstArrayView<UPrimitiveComponent*> InPrimitiveComponents);
 
 #if WITH_EDITOR
-	void OnStaticMeshChanged(UStaticMeshComponent* InStaticMeshComponent);
+	UE_DEPRECATED(5.8, "Logic moved to UE::GeometryMask::FMaskWriter")
+	UE_API void OnStaticMeshChanged(UStaticMeshComponent* InStaticMeshComponent);
 #endif
-	void OnDynamicMeshChanged(UDynamicMeshComponent* InDynamicMeshComponent);
+	UE_DEPRECATED(5.8, "Logic moved to UE::GeometryMask::FMaskWriter")
+	UE_API void OnDynamicMeshChanged(UDynamicMeshComponent* InDynamicMeshComponent);
 	
 	UPROPERTY(EditAnywhere, BlueprintReadWrite, Category="Mask", meta=(ShowOnlyInnerProperties))
 	FGeometryMaskWriteParameters Parameters;
@@ -97,18 +82,26 @@ protected:
 	UPROPERTY(EditAnywhere, BlueprintReadWrite, Category="Mask")
 	bool bWriteWhenHidden = true;
 
-	UPROPERTY(Transient, DuplicateTransient)
+	UE_DEPRECATED(5.8, "Deprecated to use MaskWriter instead")
+	UPROPERTY(Transient, meta=(DeprecatedProperty, DeprecatedMessage="Deprecated to use MaskWriter instead"))
 	int32 LastPrimitiveComponentCount;
 
-	UE_DEPRECATED(5.7, "Replaced by CachedComponents")
-	UPROPERTY(Transient, meta=(DeprecatedProperty, DeprecatedMessage="Replaced by CachedComponents"))
+	UE_DEPRECATED(5.7, "Deprecated to use MaskWriter instead")
+	UPROPERTY(Transient, meta=(DeprecatedProperty, DeprecatedMessage="Deprecated to use MaskWriter instead"))
 	TMap<FName, TWeakObjectPtr<UPrimitiveComponent>> CachedComponentsWeak;
 
 	/** Cached primitive component to mesh name, multiple components can use same mesh data */
-	UPROPERTY(Transient, DuplicateTransient)
+	UE_DEPRECATED(5.8, "Deprecated to use MaskWriter instead")
+	UPROPERTY(Transient, meta=(DeprecatedProperty, DeprecatedMessage="Deprecated to use MaskWriter instead"))
 	TMap<TWeakObjectPtr<UPrimitiveComponent>, FName> CachedComponents;
 
 	/** Cached mesh name to data (one entry per unique mesh) */
-	UPROPERTY(Transient, DuplicateTransient)
+	UE_DEPRECATED(5.8, "Deprecated to use MaskWriter instead")
+	UPROPERTY(Transient, meta=(DeprecatedProperty, DeprecatedMessage="Deprecated to use MaskWriter instead"))
 	TMap<FName, FGeometryMaskBatchElementData> CachedMeshData;
+
+	/** Handles caching and drawing the mask */
+	TSharedPtr<UE::GeometryMask::FMaskWriter> MaskWriter;
 };
+
+#undef UE_API
diff --git a/Engine/Plugins/VirtualProduction/GeometryMask/Source/GeometryMask/Public/GeometryMaskWriter.h b/Engine/Plugins/VirtualProduction/GeometryMask/Source/GeometryMask/Public/GeometryMaskWriter.h
new file mode 100644
--- /dev/null
+++ b/Engine/Plugins/VirtualProduction/GeometryMask/Source/GeometryMask/Public/GeometryMaskWriter.h
@@ -0,0 +1,65 @@
+// Copyright Epic Games, Inc. All Rights Reserved.
+
+#pragma once
+
+#include "Containers/Map.h"
+#include "GeometryMaskTypes.h"
+#include "Templates/SharedPointer.h"
+#include "UObject/NameTypes.h"
+#include "UObject/ObjectKey.h"
+
+#define UE_API GEOMETRYMASK_API
+
+class UDynamicMeshComponent;
+class UPrimitiveComponent;
+class UStaticMeshComponent;
+
+namespace UE::GeometryMask
+{
+
+/** Handles writing to the mask canvas and caching component draw data */
+struct FMaskWriter : public TSharedFromThis<FMaskWriter>
+{
+	static UE_API TSharedRef<FMaskWriter> Create();
+
+	struct FDrawParams
+	{
+		/** The canvas to write to */
+		FCanvas* Canvas = nullptr;
+		/** The actor to gather components from */
+		AActor* Actor = nullptr;
+		/** Write parameters */
+		TNotNull<const FGeometryMaskWriteParameters*> Parameters;
+		/** Whether to write the actor if it's hidden */
+		bool bWriteWhenHidden = true;
+	};
+	UE_API void DrawToCanvas(const FDrawParams& InParams);
+
+	/** Iterates each primitive that is masked */
+	void ForEachMaskPrimitive(TFunctionRef<void(TNotNull<const UPrimitiveComponent*>)> InFunc) const;
+
+private:
+	// Resets cached data, triggers rebuild
+	void ResetCachedData();
+	void UpdateCachedData(AActor* InActor);
+	void UpdateCachedStaticMeshData(TConstArrayView<UPrimitiveComponent*> InPrimitiveComponents);
+	void UpdateCachedDynamicMeshData(TConstArrayView<UPrimitiveComponent*> InPrimitiveComponents);
+
+#if WITH_EDITOR
+	void OnStaticMeshChanged(UStaticMeshComponent* InStaticMeshComponent);
+#endif
+	void OnDynamicMeshChanged(UDynamicMeshComponent* InDynamicMeshComponent);
+
+	/** Cached primitive component to the mesh object used, multiple components can use same mesh data */
+	TMap<TObjectKey<UPrimitiveComponent>, FObjectKey> CachedComponents;
+
+	/** Cached mesh object to its batch element data (one entry per unique mesh) */
+	TMap<FObjectKey, FGeometryMaskBatchElementData> CachedMeshData;
+
+	/** Cached component count per actor */
+	TMap<TObjectKey<AActor>, int32> LastPrimitiveComponentCountMap;
+};
+
+} // UE::GeometryMask
+
+#undef UE_API
diff --git a/Engine/Plugins/VirtualProduction/GeometryMask/Source/GeometryMask/Public/IGeometryMaskWriteInterface.h b/Engine/Plugins/VirtualProduction/GeometryMask/Source/GeometryMask/Public/IGeometryMaskWriteInterface.h
--- a/Engine/Plugins/VirtualProduction/GeometryMask/Source/GeometryMask/Public/IGeometryMaskWriteInterface.h
+++ b/Engine/Plugins/VirtualProduction/GeometryMask/Source/GeometryMask/Public/IGeometryMaskWriteInterface.h
@@ -6,6 +6,11 @@
 #include "UObject/Interface.h"
 #include "IGeometryMaskWriteInterface.generated.h"
 
+namespace UE::GeometryMask
+{
+	struct FMaskWriter;
+}
+
 class FCanvas;
 struct FGeometryMaskWriteParameters;
 
@@ -25,4 +30,11 @@ public:
 	virtual void SetParameters(FGeometryMaskWriteParameters& InParameters) = 0;
 	virtual void DrawToCanvas(FCanvas* InCanvas) = 0;
 	virtual FOnGeometryMaskSetCanvasNativeDelegate& OnSetCanvas() = 0;
+
+	/** Gets the mask writer for this interface */
+	virtual UE::GeometryMask::FMaskWriter* GetMaskWriter()
+	{
+		ensureMsgf(false, TEXT("5.8: Geometry Mask Write Interface does not implement 'GetMaskWriter'!"));
+		return nullptr;
+	}
 };
